Cameras: camera modifier rigs added from the player camera manager should be aware of the player controller

diff --git a/Engine/Plugins/Cameras/GameplayCameras/Source/GameplayCameras/Private/GameFramework/GameplayCamerasPlayerCameraManager.cpp b/Engine/Plugins/Cameras/GameplayCameras/Source/GameplayCameras/Private/GameFramework/GameplayCamerasPlayerCameraManager.cpp
--- a/Engine/Plugins/Cameras/GameplayCameras/Source/GameplayCameras/Private/GameFramework/GameplayCamerasPlayerCameraManager.cpp
+++ b/Engine/Plugins/Cameras/GameplayCameras/Source/GameplayCameras/Private/GameFramework/GameplayCamerasPlayerCameraManager.cpp
@@ -104,7 +104,10 @@ void AGameplayCamerasPlayerCameraManager::EnsureNullContext()
 
 	if (!NullContext)
 	{
-		NullContext = MakeShared<FCameraEvaluationContext>();
+		FCameraEvaluationContextInitializeParams InitParams;
+		InitParams.Owner = this;
+		InitParams.PlayerController = GetOwningPlayerController();
+		NullContext = MakeShared<FCameraEvaluationContext>(InitParams);
 	}
 }
 
@@ -114,8 +117,10 @@ FCameraRigInstanceID AGameplayCamerasPlayerCameraManager::StartGlobalCameraModif
 
 	if (CameraSystemEvaluator)
 	{
+		EnsureNullContext();
+
 		TSharedPtr<FCameraModifierService> CameraModifierService = CameraSystemEvaluator->FindEvaluationService<FCameraModifierService>();
-		return CameraModifierService->StartCameraModifierRig(CameraRig, ECameraRigLayer::Global, OrderKey);
+		return CameraModifierService->StartCameraModifierRig(CameraRig, NullContext.ToSharedRef(), ECameraRigLayer::Global, OrderKey);
 	}
 
 	return FCameraRigInstanceID();
@@ -127,8 +132,10 @@ FCameraRigInstanceID AGameplayCamerasPlayerCameraManager::StartVisualCameraModif
 
 	if (CameraSystemEvaluator)
 	{
+		EnsureNullContext();
+
 		TSharedPtr<FCameraModifierService> CameraModifierService = CameraSystemEvaluator->FindEvaluationService<FCameraModifierService>();
-		return CameraModifierService->StartCameraModifierRig(CameraRig, ECameraRigLayer::Visual, OrderKey);
+		return CameraModifierService->StartCameraModifierRig(CameraRig, NullContext.ToSharedRef(), ECameraRigLayer::Visual, OrderKey);
 	}
 
 	return FCameraRigInstanceID();
diff --git a/Engine/Plugins/Cameras/GameplayCameras/Source/GameplayCameras/Private/Services/CameraModifierService.cpp b/Engine/Plugins/Cameras/GameplayCameras/Source/GameplayCameras/Private/Services/CameraModifierService.cpp
--- a/Engine/Plugins/Cameras/GameplayCameras/Source/GameplayCameras/Private/Services/CameraModifierService.cpp
+++ b/Engine/Plugins/Cameras/GameplayCameras/Source/GameplayCameras/Private/Services/CameraModifierService.cpp
@@ -27,13 +27,18 @@ void FCameraModifierService::OnTeardown(const FCameraEvaluationServiceTeardownPa
 }
 
 FCameraRigInstanceID FCameraModifierService::StartCameraModifierRig(const UCameraRigAsset* CameraRig, ECameraRigLayer Layer, int32 OrderKey)
+{
+	EnsureModifierContextCreated();
+
+	return StartCameraModifierRig(CameraRig, ModifierContext.ToSharedRef(), Layer, OrderKey);
+}
+
+FCameraRigInstanceID FCameraModifierService::StartCameraModifierRig(const UCameraRigAsset* CameraRig, TSharedRef<FCameraEvaluationContext> EvaluationContext, ECameraRigLayer Layer, int32 OrderKey)
 {
 	if (ensure(Evaluator))
 	{
-		EnsureModifierContextCreated();
-
 		FActivateCameraRigParams ActivateParams;
-		ActivateParams.EvaluationContext = ModifierContext;
+		ActivateParams.EvaluationContext = EvaluationContext;
 		ActivateParams.CameraRig = CameraRig;
 		ActivateParams.Layer = Layer;
 		ActivateParams.OrderKey = FirstBlendStackOrderKey + OrderKey;
diff --git a/Engine/Plugins/Cameras/GameplayCameras/Source/GameplayCameras/Public/GameFramework/GameplayCamerasPlayerCameraManager.h b/Engine/Plugins/Cameras/GameplayCameras/Source/GameplayCameras/Public/GameFramework/GameplayCamerasPlayerCameraManager.h
--- a/Engine/Plugins/Cameras/GameplayCameras/Source/GameplayCameras/Public/GameFramework/GameplayCamerasPlayerCameraManager.h
+++ b/Engine/Plugins/Cameras/GameplayCameras/Source/GameplayCameras/Public/GameFramework/GameplayCamerasPlayerCameraManager.h
@@ -96,12 +96,15 @@ public:
 
 public:
 
+	/** Starts a camera modifier rig on the global layer. */
 	UFUNCTION(BlueprintCallable, Category="Camera")
 	FCameraRigInstanceID StartGlobalCameraModifierRig(const UCameraRigAsset* CameraRig, int32 OrderKey = 0);
 
+	/** Starts a camera modifier rig on the visual layer. */
 	UFUNCTION(BlueprintCallable, Category="Camera")
 	FCameraRigInstanceID StartVisualCameraModifierRig(const UCameraRigAsset* CameraRig, int32 OrderKey = 0);
 
+	/** Stops a camera modifier rig on previously started on the global or visual layer. */
 	UFUNCTION(BlueprintCallable, Category="Camera")
 	void StopCameraModifierRig(FCameraRigInstanceID InstanceID, bool bImmediately = false);
 
diff --git a/Engine/Plugins/Cameras/GameplayCameras/Source/GameplayCameras/Public/Services/CameraModifierService.h b/Engine/Plugins/Cameras/GameplayCameras/Source/GameplayCameras/Public/Services/CameraModifierService.h
--- a/Engine/Plugins/Cameras/GameplayCameras/Source/GameplayCameras/Public/Services/CameraModifierService.h
+++ b/Engine/Plugins/Cameras/GameplayCameras/Source/GameplayCameras/Public/Services/CameraModifierService.h
@@ -28,6 +28,9 @@ public:
 	/** Starts a new instance of the given camera rig, using a "null" context. */
 	GAMEPLAYCAMERAS_API FCameraRigInstanceID StartCameraModifierRig(const UCameraRigAsset* CameraRig, ECameraRigLayer Layer, int32 OrderKey = 0);
 
+	/** Starts a new instance of the given camera rig, using the given context. */
+	GAMEPLAYCAMERAS_API FCameraRigInstanceID StartCameraModifierRig(const UCameraRigAsset* CameraRig, TSharedRef<FCameraEvaluationContext> EvaluationContext, ECameraRigLayer Layer, int32 OrderKey = 0);
+
 	/** Stops a running instance of a camera rig. */
 	GAMEPLAYCAMERAS_API void StopCameraModifierRig(FCameraRigInstanceID CameraRigID, bool bImmediately = false);
 
