Cameras: big pass on camera shakes

- Add shake instance IDs
- Use these IDs to allow stopping ongoing shakes
- Support both immediate and blend-out stopping, similar to the persistent blend stacks
- Improve the API for shake nodes, including the introduction of a "shake delta" structure that makes it easy to add transform shake without having to bother with play spaces and scales, which are applied at the end
- Fix API for shake duration and fix how the envelope shake node handles duration

diff --git a/Engine/Plugins/Cameras/GameplayCameras/Source/GameplayCameras/Private/Core/ShakeCameraNode.cpp b/Engine/Plugins/Cameras/GameplayCameras/Source/GameplayCameras/Private/Core/ShakeCameraNode.cpp
--- a/Engine/Plugins/Cameras/GameplayCameras/Source/GameplayCameras/Private/Core/ShakeCameraNode.cpp
+++ b/Engine/Plugins/Cameras/GameplayCameras/Source/GameplayCameras/Private/Core/ShakeCameraNode.cpp
@@ -2,11 +2,82 @@
 
 #include "Core/ShakeCameraNode.h"
 
+#include "Core/CameraPose.h"
+#include "Core/PostProcessSettingsCollection.h"
+#include "Engine/Engine.h"
+#include "IXRTrackingSystem.h"
+
 #include UE_INLINE_GENERATED_CPP_BY_NAME(ShakeCameraNode)
 
 namespace UE::Cameras
 {
 
+void FCameraNodeShakeDelta::Combine(const FCameraNodeShakeDelta& Other, float OtherScale)
+{
+	Location += OtherScale * Other.Location;
+	Rotation += OtherScale * Other.Rotation;
+	FieldOfView += OtherScale * Other.FieldOfView;
+}
+
+void FCameraNodeShakeResult::ApplyDelta(const FCameraNodeShakeParams& Params)
+{
+	FVector3d FinalDeltaLocation = ShakeDelta.Location;
+	FRotator3d FinalDeltaRotation = ShakeDelta.Rotation;
+	float FinalDeltaFieldOfView = ShakeDelta.FieldOfView;
+
+	// Apply the shake scale.
+	if (Params.ShakeScale != 1.f)
+	{
+		FinalDeltaLocation *= Params.ShakeScale;
+		FinalDeltaRotation *= Params.ShakeScale;
+		FinalDeltaFieldOfView *= Params.ShakeScale;
+	}
+
+	// Apply the same limits as the vanilla Blueprint camera shakes:
+	// - Don't allow shake to flip pitch past vertical, if not using a headset.
+	// - If using a headset, we can't limit the camera locked to the player's head.
+	if (!FinalDeltaRotation.IsZero() && (!GEngine->XRSystem.IsValid() || !GEngine->XRSystem->IsHeadTrackingAllowed()))
+	{
+		// Find normalized result when combined, and remove any offset that would push it past the limit.
+		const float NormalizedInputPitch = FRotator::NormalizeAxis(ShakenResult.CameraPose.GetRotation().Pitch);
+		const float NormalizedOutputPitchOffset = FRotator::NormalizeAxis(FinalDeltaRotation.Pitch);
+		FinalDeltaRotation.Pitch = FMath::ClampAngle(NormalizedInputPitch + NormalizedOutputPitchOffset, -89.9f, 89.9f) - NormalizedInputPitch;
+	}
+
+	// Apply location and rotation.
+	FCameraPose& CameraPose = ShakenResult.CameraPose;
+	const FRotationMatrix CameraRotation(CameraPose.GetRotation());
+	if (Params.PlaySpace == ECameraShakePlaySpace::CameraLocal)
+	{
+		// Convert the delta rotation/location into the camera's rotation space.
+		const FVector3d ShakenPoseLocation = CameraPose.GetLocation() + CameraRotation.TransformVector(FinalDeltaLocation);
+		CameraPose.SetLocation(ShakenPoseLocation);
+
+		const FRotator3d ShakenPoseRotation = (FRotationMatrix(FinalDeltaRotation) * CameraRotation).Rotator();
+		CameraPose.SetRotation(ShakenPoseRotation);
+	}
+	else if (Params.PlaySpace == ECameraShakePlaySpace::UserDefined)
+	{
+		const FMatrix& UserPlaySpaceMatrix = Params.UserPlaySpaceMatrix;
+
+		const FVector3d ShakenPoseLocation = CameraPose.GetLocation() + UserPlaySpaceMatrix.TransformVector(FinalDeltaLocation);
+		CameraPose.SetLocation(ShakenPoseLocation);
+
+		const FMatrix WorldDeltaRotation = UserPlaySpaceMatrix.Inverse() * FRotationMatrix(FinalDeltaRotation) * UserPlaySpaceMatrix;
+		const FMatrix LocalDeltaRotation = CameraRotation.Inverse() * WorldDeltaRotation;
+		const FRotator3d ShakenPoseRotation = (LocalDeltaRotation * CameraRotation).Rotator();
+		CameraPose.SetRotation(ShakenPoseRotation);
+	}
+
+	// Apply field of view.
+	if (FinalDeltaFieldOfView != 0.f)
+	{
+		const double EffectiveFieldOfView = CameraPose.GetEffectiveFieldOfView();
+		CameraPose.SetFieldOfView(EffectiveFieldOfView + FinalDeltaFieldOfView);
+		CameraPose.SetFocalLength(-1.f);
+	}
+}
+
 UE_DEFINE_CAMERA_NODE_EVALUATOR(FShakeCameraNodeEvaluator)
 
 void FShakeCameraNodeEvaluator::ShakeResult(const FCameraNodeShakeParams& Params, FCameraNodeShakeResult& OutResult)
diff --git a/Engine/Plugins/Cameras/GameplayCameras/Source/GameplayCameras/Private/GameFramework/GameplayCamerasPlayerCameraManager.cpp b/Engine/Plugins/Cameras/GameplayCameras/Source/GameplayCameras/Private/GameFramework/GameplayCamerasPlayerCameraManager.cpp
--- a/Engine/Plugins/Cameras/GameplayCameras/Source/GameplayCameras/Private/GameFramework/GameplayCamerasPlayerCameraManager.cpp
+++ b/Engine/Plugins/Cameras/GameplayCameras/Source/GameplayCameras/Private/GameFramework/GameplayCamerasPlayerCameraManager.cpp
@@ -19,6 +19,7 @@
 #include "GameFramework/ViewTargetTransitionParamsBlendNode.h"
 #include "GameplayCamerasSettings.h"
 #include "Services/CameraModifierService.h"
+#include "Services/CameraShakeService.h"
 #include "UObject/Package.h"
 #include "UObject/UObjectGlobals.h"
 
@@ -152,6 +153,49 @@ void AGameplayCamerasPlayerCameraManager::StopCameraModifierRig(FCameraRigInstan
 	}
 }
 
+FCameraShakeInstanceID AGameplayCamerasPlayerCameraManager::StartCameraShakeAsset(const UCameraShakeAsset* CameraShake, float ShakeScale, ECameraShakePlaySpace PlaySpace, FRotator UserPlaySpaceRotation)
+{
+	using namespace UE::Cameras;
+
+	if (CameraSystemEvaluator)
+	{
+		FStartCameraShakeParams Params;
+		Params.CameraShake = CameraShake;
+		Params.ShakeScale = ShakeScale;
+		Params.PlaySpace = PlaySpace;
+		Params.UserPlaySpaceRotation = UserPlaySpaceRotation;
+
+		TSharedPtr<FCameraShakeService> CameraShakeService = CameraSystemEvaluator->FindEvaluationService<FCameraShakeService>();
+		return CameraShakeService->StartCameraShake(Params);
+	}
+
+	return FCameraShakeInstanceID();
+}
+
+bool AGameplayCamerasPlayerCameraManager::IsCameraShakeAssetPlaying(FCameraShakeInstanceID InInstanceID) const
+{
+	using namespace UE::Cameras;
+
+	if (CameraSystemEvaluator)
+	{
+		TSharedPtr<FCameraShakeService> CameraShakeService = CameraSystemEvaluator->FindEvaluationService<FCameraShakeService>();
+		return CameraShakeService->IsCameraShakePlaying(InInstanceID);
+	}
+	return false;
+}
+
+bool AGameplayCamerasPlayerCameraManager::StopCameraShakeAsset(FCameraShakeInstanceID InInstanceID, bool bImmediately)
+{
+	using namespace UE::Cameras;
+
+	if (CameraSystemEvaluator)
+	{
+		TSharedPtr<FCameraShakeService> CameraShakeService = CameraSystemEvaluator->FindEvaluationService<FCameraShakeService>();
+		return CameraShakeService->StopCameraShake(InInstanceID, bImmediately);
+	}
+	return false;
+}
+
 void AGameplayCamerasPlayerCameraManager::InitializeFor(APlayerController* PlayerController)
 {
 	using namespace UE::Cameras;
diff --git a/Engine/Plugins/Cameras/GameplayCameras/Source/GameplayCameras/Private/Nodes/Shakes/CameraShakeCameraNode.cpp b/Engine/Plugins/Cameras/GameplayCameras/Source/GameplayCameras/Private/Nodes/Shakes/CameraShakeCameraNode.cpp
--- a/Engine/Plugins/Cameras/GameplayCameras/Source/GameplayCameras/Private/Nodes/Shakes/CameraShakeCameraNode.cpp
+++ b/Engine/Plugins/Cameras/GameplayCameras/Source/GameplayCameras/Private/Nodes/Shakes/CameraShakeCameraNode.cpp
@@ -105,6 +105,7 @@ void FCameraShakeCameraNodeEvaluator::OnRun(const FCameraNodeEvaluationParams& P
 		FCameraNodeShakeParams ShakeParams(Params);
 		FCameraNodeShakeResult ShakeResult(OutResult);
 		CameraShakeRootEvaluator->ShakeResult(ShakeParams, ShakeResult);
+		ShakeResult.ApplyDelta(ShakeParams);
 	}
 	// If evaluating the shake later in the visual layer, keep requesting that the shake service
 	// maintains this shake alive. We already put the shake's parameters inside the variable table so
diff --git a/Engine/Plugins/Cameras/GameplayCameras/Source/GameplayCameras/Private/Nodes/Shakes/EnvelopeShakeCameraNode.cpp b/Engine/Plugins/Cameras/GameplayCameras/Source/GameplayCameras/Private/Nodes/Shakes/EnvelopeShakeCameraNode.cpp
--- a/Engine/Plugins/Cameras/GameplayCameras/Source/GameplayCameras/Private/Nodes/Shakes/EnvelopeShakeCameraNode.cpp
+++ b/Engine/Plugins/Cameras/GameplayCameras/Source/GameplayCameras/Private/Nodes/Shakes/EnvelopeShakeCameraNode.cpp
@@ -105,14 +105,30 @@ void FEnvelopeShakeCameraNodeEvaluator::OnShakeResult(const FCameraNodeShakePara
 		Alpha = SmoothStep((TotalTime - CurrentTime) / EaseOutTime);
 	}
 
-	OutResult.ShakeTimeLeft = (TotalTime - CurrentTime);
+	const float EnvelopeTimeLeft = (TotalTime - CurrentTime);
 	{
 		FCameraNodeShakeParams ChildParams(Params);
 		ChildParams.ShakeScale *= Alpha;
 
-		ShakeEvaluator->ShakeResult(ChildParams, OutResult);
+		FCameraNodeShakeResult ChildResult(OutResult.ShakenResult);
+
+		ShakeEvaluator->ShakeResult(ChildParams, ChildResult);
+
+		OutResult.ShakeDelta.Combine(ChildResult.ShakeDelta, ChildParams.ShakeScale);
+		OutResult.ShakeTimeLeft = ChildResult.ShakeTimeLeft;
+	}
+
+	if (OutResult.ShakeTimeLeft >= 0.f)
+	{
+		// The underlying shake has a finite duration, so we end when it ends, or when we
+		// reach our total time, whichever happens first.
+		OutResult.ShakeTimeLeft = FMath::Clamp(OutResult.ShakeTimeLeft, 0, EnvelopeTimeLeft);
+	}
+	else
+	{
+		// The underlying shake is infinite, so we drive the duration.
+		OutResult.ShakeTimeLeft = EnvelopeTimeLeft;
 	}
-	OutResult.ShakeTimeLeft = FMath::Clamp(OutResult.ShakeTimeLeft, 0, TotalTime - CurrentTime);
 }
 
 void FEnvelopeShakeCameraNodeEvaluator::OnRestartShake(const FCameraNodeShakeRestartParams& Params)
diff --git a/Engine/Plugins/Cameras/GameplayCameras/Source/GameplayCameras/Private/Nodes/Shakes/PerlinNoiseLocationShakeCameraNode.cpp b/Engine/Plugins/Cameras/GameplayCameras/Source/GameplayCameras/Private/Nodes/Shakes/PerlinNoiseLocationShakeCameraNode.cpp
--- a/Engine/Plugins/Cameras/GameplayCameras/Source/GameplayCameras/Private/Nodes/Shakes/PerlinNoiseLocationShakeCameraNode.cpp
+++ b/Engine/Plugins/Cameras/GameplayCameras/Source/GameplayCameras/Private/Nodes/Shakes/PerlinNoiseLocationShakeCameraNode.cpp
@@ -112,9 +112,7 @@ void FPerlinNoiseLocationShakeCameraNodeEvaluator::OnRun(const FCameraNodeEvalua
 
 void FPerlinNoiseLocationShakeCameraNodeEvaluator::OnShakeResult(const FCameraNodeShakeParams& Params, FCameraNodeShakeResult& OutResult)
 {
-	FVector3d Location = OutResult.ShakenResult.CameraPose.GetLocation();
-	Location += NoiseValue;
-	OutResult.ShakenResult.CameraPose.SetLocation(Location);
+	OutResult.ShakeDelta.Location += NoiseValue;
 	OutResult.ShakeTimeLeft = -1.f;
 }
 
diff --git a/Engine/Plugins/Cameras/GameplayCameras/Source/GameplayCameras/Private/Nodes/Shakes/PerlinNoiseRotationShakeCameraNode.cpp b/Engine/Plugins/Cameras/GameplayCameras/Source/GameplayCameras/Private/Nodes/Shakes/PerlinNoiseRotationShakeCameraNode.cpp
--- a/Engine/Plugins/Cameras/GameplayCameras/Source/GameplayCameras/Private/Nodes/Shakes/PerlinNoiseRotationShakeCameraNode.cpp
+++ b/Engine/Plugins/Cameras/GameplayCameras/Source/GameplayCameras/Private/Nodes/Shakes/PerlinNoiseRotationShakeCameraNode.cpp
@@ -111,9 +111,7 @@ void FPerlinNoiseRotationShakeCameraNodeEvaluator::OnRun(const FCameraNodeEvalua
 
 void FPerlinNoiseRotationShakeCameraNodeEvaluator::OnShakeResult(const FCameraNodeShakeParams& Params, FCameraNodeShakeResult& OutResult)
 {
-	FRotator3d Rotation = OutResult.ShakenResult.CameraPose.GetRotation();
-	Rotation += NoiseValue;
-	OutResult.ShakenResult.CameraPose.SetRotation(Rotation);
+	OutResult.ShakeDelta.Rotation += NoiseValue;
 	OutResult.ShakeTimeLeft = -1.f;
 }
 
diff --git a/Engine/Plugins/Cameras/GameplayCameras/Source/GameplayCameras/Private/Services/CameraShakeService.cpp b/Engine/Plugins/Cameras/GameplayCameras/Source/GameplayCameras/Private/Services/CameraShakeService.cpp
--- a/Engine/Plugins/Cameras/GameplayCameras/Source/GameplayCameras/Private/Services/CameraShakeService.cpp
+++ b/Engine/Plugins/Cameras/GameplayCameras/Source/GameplayCameras/Private/Services/CameraShakeService.cpp
@@ -11,6 +11,8 @@
 #include "Core/CameraSystemEvaluator.h"
 #include "Core/RootCameraNode.h"
 #include "Core/ShakeCameraNode.h"
+#include "Nodes/Blends/InterruptedBlendCameraNode.h"
+#include "Nodes/Blends/ReverseBlendCameraNode.h"
 #include "UObject/ObjectMacros.h"
 
 #include UE_INLINE_GENERATED_CPP_BY_NAME(CameraShakeService)
@@ -24,7 +26,9 @@ class FCameraShakeServiceCameraNodeEvaluator : public FCameraNodeEvaluator
 
 public:
 
-	void StartCameraShake(const FStartCameraShakeParams& Params);
+	FCameraShakeInstanceID StartCameraShake(const FStartCameraShakeParams& Params);
+	bool IsCameraShakePlaying(FCameraShakeInstanceID InInstanceID) const;
+	bool StopCameraShake(FCameraShakeInstanceID InInstanceID, bool bImmediately);
 	void RequestCameraShakeThisFrame(const FStartCameraShakeParams& Params);
 
 protected:
@@ -45,27 +49,39 @@ private:
 		const UCameraShakeAsset* CameraShake,
 		TSharedPtr<const FCameraEvaluationContext> EvaluationContext);
 
+	bool InitializeEntryBlendOut(FShakeEntry& Entry);
+
 	void PopEntry(int32 EntryIndex);
 
 private:
 
+	enum class EBlendStatus
+	{
+		None,
+		BlendIn,
+		BlendOut
+	};
+
 	struct FShakeEntry
 	{
+		FCameraShakeInstanceID EntryID;
 		TWeakPtr<const FCameraEvaluationContext> EvaluationContext;
 		TObjectPtr<const UCameraShakeAsset> CameraShake;
 		FCameraNodeEvaluatorStorage EvaluatorStorage;
-		FBlendCameraNodeEvaluator* BlendInEvaluator = nullptr;
-		FBlendCameraNodeEvaluator* BlendOutEvaluator = nullptr;
+		FBlendCameraNodeEvaluator* BlendEvaluator = nullptr;
 		FShakeCameraNodeEvaluator* RootEvaluator = nullptr;
 		FCameraNodeEvaluatorHierarchy EvaluatorHierarchy;
 		FCameraNodeEvaluationResult Result;
 		float CurrentTime = 0.f;
 		float ShakeScale = 1.f;
+		EBlendStatus BlendStatus = EBlendStatus::None;
 		ECameraShakePlaySpace PlaySpace = ECameraShakePlaySpace::CameraLocal;
-		FMatrix UserPlaySpaceMatrix;
+		FMatrix UserPlaySpaceMatrix = FMatrix::Identity;
 		uint8 NumRequests = 0;
 		bool bPersistentRequest = false;
 		bool bIsFirstFrame = false;
+		bool bIsBlendFull = false;
+		bool bIsBlendFinished = false;
 	};
 
 	FCameraSystemEvaluator* OwningEvaluator = nullptr;
@@ -73,6 +89,8 @@ private:
 	const FCameraVariableTable* BlendedParameters = nullptr;;
 
 	TArray<FShakeEntry> Entries;
+
+	uint32 NextEntryID = 0;
 };
 
 UE_DEFINE_CAMERA_NODE_EVALUATOR(FCameraShakeServiceCameraNodeEvaluator)
@@ -87,11 +105,11 @@ FCameraNodeEvaluatorChildrenView FCameraShakeServiceCameraNodeEvaluator::OnGetCh
 	return ChildrenView;
 }
 
-void FCameraShakeServiceCameraNodeEvaluator::StartCameraShake(const FStartCameraShakeParams& Params)
+FCameraShakeInstanceID FCameraShakeServiceCameraNodeEvaluator::StartCameraShake(const FStartCameraShakeParams& Params)
 {
 	if (!Params.CameraShake)
 	{
-		return;
+		return FCameraShakeInstanceID();
 	}
 
 	// If this shake wants to only have a single instance active at a time, look for a running
@@ -111,7 +129,7 @@ void FCameraShakeServiceCameraNodeEvaluator::StartCameraShake(const FStartCamera
 		{
 			FCameraNodeShakeRestartParams RestartParams;
 			ExistingEntry->RootEvaluator->RestartShake(RestartParams);
-			return;
+			return ExistingEntry->EntryID;
 		}
 	}
 
@@ -119,7 +137,44 @@ void FCameraShakeServiceCameraNodeEvaluator::StartCameraShake(const FStartCamera
 	if (NewEntry)
 	{
 		NewEntry->bPersistentRequest = true;
+		return NewEntry->EntryID;
 	}
+	return FCameraShakeInstanceID();
+}
+
+bool FCameraShakeServiceCameraNodeEvaluator::IsCameraShakePlaying(FCameraShakeInstanceID InInstanceID) const
+{
+	return Entries.ContainsByPredicate(
+			[InInstanceID](const FShakeEntry& Entry)
+			{
+				return Entry.EntryID == InInstanceID;
+			});
+}
+
+bool FCameraShakeServiceCameraNodeEvaluator::StopCameraShake(FCameraShakeInstanceID InInstanceID, bool bImmediately)
+{
+	const int32 EntryIndex = Entries.IndexOfByPredicate(
+			[InInstanceID](const FShakeEntry& Entry)
+			{
+				return Entry.EntryID == InInstanceID;
+			});
+	if (EntryIndex != INDEX_NONE)
+	{
+		if (bImmediately)
+		{
+			PopEntry(EntryIndex);
+		}
+		else
+		{
+			const bool bHasBlendOut = InitializeEntryBlendOut(Entries[EntryIndex]);
+			if (!bHasBlendOut)
+			{
+				PopEntry(EntryIndex);
+			}
+		}
+		return true;
+	}
+	return false;
 }
 
 void FCameraShakeServiceCameraNodeEvaluator::RequestCameraShakeThisFrame(const FStartCameraShakeParams& Params)
@@ -164,7 +219,11 @@ FCameraShakeServiceCameraNodeEvaluator::FShakeEntry* FCameraShakeServiceCameraNo
 
 	NewEntry.ShakeScale = Params.ShakeScale;
 	NewEntry.PlaySpace = Params.PlaySpace;
-	NewEntry.UserPlaySpaceMatrix = Params.UserPlaySpaceMatrix;
+	NewEntry.UserPlaySpaceMatrix = FMatrix::Identity;
+	if (Params.PlaySpace == ECameraShakePlaySpace::UserDefined)
+	{
+		NewEntry.UserPlaySpaceMatrix = FRotationMatrix(Params.UserPlaySpaceRotation);
+	}
 
 	const int32 AddedIndex = Entries.Add(MoveTemp(NewEntry));
 	return &Entries[AddedIndex];
@@ -222,8 +281,14 @@ void FCameraShakeServiceCameraNodeEvaluator::OnRun(const FCameraNodeEvaluationPa
 
 		// Update timing.
 		Entry.CurrentTime += Params.DeltaTime;
+
+		// Run the blend.
+		if (Entry.BlendEvaluator)
+		{
+			Entry.BlendEvaluator->Run(CurParams, CurResult);
+		}
 		
-		// Run the shake!
+		// Run the shake and apply it.
 		float CurTimeLeft = 0.f;
 		if (Entry.RootEvaluator)
 		{
@@ -238,6 +303,8 @@ void FCameraShakeServiceCameraNodeEvaluator::OnRun(const FCameraNodeEvaluationPa
 
 			Entry.RootEvaluator->ShakeResult(ShakeParams, ShakeResult);
 
+			ShakeResult.ApplyDelta(ShakeParams);
+
 			CurTimeLeft = ShakeResult.ShakeTimeLeft;
 		}
 
@@ -245,37 +312,57 @@ void FCameraShakeServiceCameraNodeEvaluator::OnRun(const FCameraNodeEvaluationPa
 		Entry.NumRequests = 0;
 
 		// If it says it's finished, schedule it for removal.
+		// (note that negative time left means the shake should play indefinitely)
 		if (CurTimeLeft == 0)
 		{
 			EntriesToRemove.Add(Index);
 			continue;
 		}
 
-		// Update blends.
-		if (Entry.CameraShake->BlendIn && 
-				Entry.CurrentTime < Entry.CameraShake->BlendIn->BlendTime.GetValue(OutResult.VariableTable) &&
-				ensure(Entry.BlendInEvaluator))
+		// Check if we need to start blending out.
+		if (Entry.BlendStatus != EBlendStatus::BlendOut &&
+				Entry.CameraShake->BlendOut &&
+				CurTimeLeft >= 0.f && 
+				CurTimeLeft < Entry.CameraShake->BlendOut->BlendTime.GetValue(CurResult.VariableTable))
 		{
-			Entry.BlendInEvaluator->Run(CurParams, CurResult);
-
-			FCameraNodeBlendParams BlendParams(Params, Entry.Result);
-			FCameraNodeBlendResult BlendResult(OutResult);
-			Entry.BlendInEvaluator->BlendResults(BlendParams, BlendResult);
+			const bool bHasBlendOut = InitializeEntryBlendOut(Entry);
+			if (!bHasBlendOut)
+			{
+				EntriesToRemove.Add(Index);
+				continue;
+			}
 		}
-		else if (Entry.CameraShake->BlendOut &&
-				CurTimeLeft >= 0.f && CurTimeLeft < Entry.CameraShake->BlendOut->BlendTime.GetValue(OutResult.VariableTable) &&
-				ensure(Entry.BlendOutEvaluator))
-		{
-			Entry.BlendOutEvaluator->Run(CurParams, CurResult);
 
-			FCameraNodeBlendParams BlendParams(Params, Entry.Result);
+		// Apply blending.
+		if (Entry.BlendEvaluator)
+		{
+			FCameraNodeBlendParams BlendParams(Params, CurResult);
 			FCameraNodeBlendResult BlendResult(OutResult);
-			Entry.BlendOutEvaluator->BlendResults(BlendParams, BlendResult);
+			Entry.BlendEvaluator->BlendResults(BlendParams, BlendResult);
+
+			Entry.bIsBlendFull = BlendResult.bIsBlendFull;
+			Entry.bIsBlendFinished = BlendResult.bIsBlendFinished;
 		}
 		else
 		{
 			OutResult.OverrideAll(CurResult);
 		}
+
+		// Update blend status.
+		if (Entry.BlendStatus == EBlendStatus::BlendIn)
+		{
+			if (Entry.bIsBlendFull && Entry.bIsBlendFinished)
+			{
+				Entry.BlendStatus = EBlendStatus::None;
+			}
+		}
+		else if (Entry.BlendStatus == EBlendStatus::BlendOut)
+		{
+			if (Entry.bIsBlendFull && Entry.bIsBlendFinished)
+			{
+				EntriesToRemove.Add(Index);
+			}
+		}
 	}
 
 	// Remove any finished shakes.
@@ -296,7 +383,7 @@ void FCameraShakeServiceCameraNodeEvaluator::InitializeEntry(
 	BuildParams.AllocationInfo = &CameraShake->AllocationInfo.EvaluatorInfo;
 	FCameraNodeEvaluator* RootEvaluator = NewEntry.EvaluatorStorage.BuildEvaluatorTree(BuildParams);
 
-	// Generate the blend-in and blend-out evaluators.
+	// Generate the blend-in evaluator.
 	FBlendCameraNodeEvaluator* BlendInEvaluator = nullptr;
 	if (CameraShake->BlendIn)
 	{
@@ -304,16 +391,6 @@ void FCameraShakeServiceCameraNodeEvaluator::InitializeEntry(
 		BlendBuildParams.RootCameraNode = CameraShake->BlendIn;
 		BlendInEvaluator = NewEntry.EvaluatorStorage.BuildEvaluatorTree(BlendBuildParams)->CastThis<FBlendCameraNodeEvaluator>();
 	}
-	FBlendCameraNodeEvaluator* BlendOutEvaluator = nullptr;
-	if (CameraShake->BlendOut)
-	{
-		FCameraNodeEvaluatorTreeBuildParams BlendBuildParams;
-		BlendBuildParams.RootCameraNode = CameraShake->BlendOut;
-		BlendOutEvaluator = NewEntry.EvaluatorStorage.BuildEvaluatorTree(BlendBuildParams)->CastThis<FBlendCameraNodeEvaluator>();
-
-		const bool bReversed = BlendOutEvaluator->SetReversed(true);
-		ensure(bReversed);  // TODO: if the blend can't play in reverse, wrap it in a reverse blend evaluator.
-	}
 
 	// Allocate variable table and context data table.
 	NewEntry.Result.VariableTable.Initialize(CameraShake->AllocationInfo.VariableTableInfo);
@@ -325,6 +402,14 @@ void FCameraShakeServiceCameraNodeEvaluator::InitializeEntry(
 	NewEntry.Result.ContextDataTable.OverrideAll(ContextResult.ContextDataTable);
 
 	// Initialize the node evaluators.
+	if (BlendInEvaluator)
+	{
+		FCameraNodeEvaluatorInitializeParams BlendInInitParams;
+		BlendInInitParams.Evaluator = OwningEvaluator;
+		BlendInInitParams.EvaluationContext = EvaluationContext;
+		BlendInInitParams.Layer = ECameraRigLayer::Visual;
+		BlendInEvaluator->Initialize(BlendInInitParams, NewEntry.Result);
+	}
 	if (RootEvaluator)
 	{
 		FCameraNodeEvaluatorInitializeParams InitParams(&NewEntry.EvaluatorHierarchy);
@@ -335,10 +420,13 @@ void FCameraShakeServiceCameraNodeEvaluator::InitializeEntry(
 	}
 
 	// Wrap up!
+	NewEntry.EntryID = FCameraShakeInstanceID(NextEntryID++);
 	NewEntry.EvaluationContext = EvaluationContext;
 	NewEntry.CameraShake = CameraShake;
-	NewEntry.BlendInEvaluator = BlendInEvaluator;
-	NewEntry.BlendOutEvaluator = BlendOutEvaluator;
+	NewEntry.BlendEvaluator = BlendInEvaluator;
+	NewEntry.BlendStatus = (BlendInEvaluator != nullptr ? EBlendStatus::BlendIn : EBlendStatus::None);
+	NewEntry.bIsBlendFull = (BlendInEvaluator == nullptr);
+	NewEntry.bIsBlendFinished = (BlendInEvaluator == nullptr);
 	NewEntry.bIsFirstFrame = true;
 	if (RootEvaluator)
 	{
@@ -346,6 +434,65 @@ void FCameraShakeServiceCameraNodeEvaluator::InitializeEntry(
 	}
 }
 
+bool FCameraShakeServiceCameraNodeEvaluator::InitializeEntryBlendOut(FShakeEntry& Entry)
+{
+	if (const USimpleFixedTimeBlendCameraNode* BlendOut = Entry.CameraShake->BlendOut)
+	{
+		// Swap the blend-in evaluator on this entry with a blend-out one.
+		if (Entry.BlendStatus != EBlendStatus::BlendOut)
+		{
+			FCameraNodeEvaluatorBuilder BlendOutBuilder(Entry.EvaluatorStorage);
+			FCameraNodeEvaluatorBuildParams BlendOutBuildParams(BlendOutBuilder);
+			FBlendCameraNodeEvaluator* BlendOutEvaluator = BlendOutBuildParams.BuildEvaluatorAs<FBlendCameraNodeEvaluator>(BlendOut);
+
+			FCameraNodeEvaluatorInitializeParams BlendOutInitParams;
+			BlendOutInitParams.Evaluator = OwningEvaluator;
+			BlendOutInitParams.EvaluationContext = Entry.EvaluationContext.Pin();
+			BlendOutInitParams.Layer = ECameraRigLayer::Visual;
+			BlendOutEvaluator->Initialize(BlendOutInitParams, Entry.Result);
+
+			// Reverse this blend so it plays as a blend-out. Also, see if we are going to 
+			// interrupt an ongoing blend-in... if so, give a chance for the blend-out to
+			// start at an "equivalent spot".
+			if (!BlendOutEvaluator->SetReversed(true))
+			{
+				BlendOutEvaluator = Entry.EvaluatorStorage.BuildEvaluator<FReverseBlendCameraNodeEvaluator>(BlendOutEvaluator);
+			}
+			if (Entry.BlendStatus == EBlendStatus::BlendIn && ensure(Entry.BlendEvaluator))
+			{
+				FBlendCameraNodeEvaluator* OngoingBlend = Entry.BlendEvaluator;
+
+				FCameraNodeBlendInterruptionParams InterruptionParams;
+				InterruptionParams.InterruptedBlend = OngoingBlend;
+				if (!BlendOutEvaluator->InitializeFromInterruption(InterruptionParams))
+				{
+					BlendOutEvaluator = Entry.EvaluatorStorage.BuildEvaluator<FInterruptedBlendCameraNodeEvaluator>(BlendOutEvaluator, OngoingBlend);
+				}
+			}
+			// Note: neither the reverse or interrupted blends need initialization, but
+			// technically we're missing calling it on them.
+			Entry.BlendEvaluator = BlendOutEvaluator;
+
+			Entry.BlendStatus = EBlendStatus::BlendOut;
+			Entry.bIsBlendFinished = false;
+			Entry.bIsBlendFull = false;
+		}
+		// else: we were already blending out, so let this continue.
+
+		return true;
+	}
+	else
+	{
+		// No blend out, just stop.
+		Entry.BlendEvaluator = nullptr;
+		Entry.BlendStatus = EBlendStatus::BlendOut;
+		Entry.bIsBlendFull = true;
+		Entry.bIsBlendFinished = true;
+
+		return false;
+	}
+}
+
 void FCameraShakeServiceCameraNodeEvaluator::PopEntry(int32 EntryIndex)
 {
 	if (!ensure(Entries.IsValidIndex(EntryIndex)))
@@ -372,14 +519,34 @@ void FCameraShakeService::OnTeardown(const FCameraEvaluationServiceTeardownParam
 	Evaluator = nullptr;
 }
 
-void FCameraShakeService::StartCameraShake(const FStartCameraShakeParams& Params)
+FCameraShakeInstanceID FCameraShakeService::StartCameraShake(const FStartCameraShakeParams& Params)
 {
 	EnsureShakeContextCreated();
 
 	if (ensure(ShakeEvaluator))
 	{
-		ShakeEvaluator->StartCameraShake(Params);
+		return ShakeEvaluator->StartCameraShake(Params);
+	}
+
+	return FCameraShakeInstanceID();
+}
+
+bool FCameraShakeService::IsCameraShakePlaying(FCameraShakeInstanceID InInstanceID) const
+{
+	if (ShakeEvaluator)
+	{
+		return ShakeEvaluator->IsCameraShakePlaying(InInstanceID);
+	}
+	return false;
+}
+
+bool FCameraShakeService::StopCameraShake(FCameraShakeInstanceID InInstanceID, bool bImmediately)
+{
+	if (ShakeEvaluator)
+	{
+		return ShakeEvaluator->StopCameraShake(InInstanceID, bImmediately);
 	}
+	return false;
 }
 
 void FCameraShakeService::RequestCameraShakeThisFrame(const FStartCameraShakeParams& Params)
diff --git a/Engine/Plugins/Cameras/GameplayCameras/Source/GameplayCameras/Public/Core/CameraShakeInstanceID.h b/Engine/Plugins/Cameras/GameplayCameras/Source/GameplayCameras/Public/Core/CameraShakeInstanceID.h
new file mode 100644
--- /dev/null
+++ b/Engine/Plugins/Cameras/GameplayCameras/Source/GameplayCameras/Public/Core/CameraShakeInstanceID.h
@@ -0,0 +1,59 @@
+// Copyright Epic Games, Inc. All Rights Reserved.
+
+#pragma once
+
+#include "CoreTypes.h"
+
+#include "CameraShakeInstanceID.generated.h"
+
+/**
+ * A unique instance ID for a running, or previously running, camera shake.
+ */
+USTRUCT(BlueprintType)
+struct FCameraShakeInstanceID
+{
+	GENERATED_BODY()
+
+public:
+
+	FCameraShakeInstanceID() : Value(INVALID) {}
+
+	FCameraShakeInstanceID(uint32 InValue) 
+		: Value(InValue)
+	{}
+
+	bool IsValid() const { return Value != INVALID; }
+
+	explicit operator bool() const { return IsValid(); }
+
+public:
+
+	friend bool operator==(FCameraShakeInstanceID A, FCameraShakeInstanceID B)
+	{
+		return A.Value == B.Value;
+	}
+
+	friend bool operator!=(FCameraShakeInstanceID A, FCameraShakeInstanceID B)
+	{
+		return !(A == B);
+	}
+
+	friend uint32 GetTypeHash(FCameraShakeInstanceID In)
+	{
+		return In.Value;
+	}
+
+	friend FArchive& operator<< (FArchive& Ar, FCameraShakeInstanceID& In)
+	{
+		Ar << In.Value;
+		return Ar;
+	}
+
+private:
+
+	static const uint32 INVALID = uint32(-1);
+
+	UPROPERTY()
+	uint32 Value;
+};
+
diff --git a/Engine/Plugins/Cameras/GameplayCameras/Source/GameplayCameras/Public/Core/ShakeCameraNode.h b/Engine/Plugins/Cameras/GameplayCameras/Source/GameplayCameras/Public/Core/ShakeCameraNode.h
--- a/Engine/Plugins/Cameras/GameplayCameras/Source/GameplayCameras/Public/Core/ShakeCameraNode.h
+++ b/Engine/Plugins/Cameras/GameplayCameras/Source/GameplayCameras/Public/Core/ShakeCameraNode.h
@@ -36,7 +36,27 @@ struct FCameraNodeShakeParams
 	/** The play space to modify the result by */
 	ECameraShakePlaySpace PlaySpace = ECameraShakePlaySpace::CameraLocal;
 	/** The custom space to use for the shake. Only used when PlaySpace is UserDefined. */
-	FMatrix UserPlaySpaceMatrix;
+	FMatrix UserPlaySpaceMatrix = FMatrix::Identity;
+};
+
+/**
+ * A structure that describes a group of offset values representing a shake.
+ */
+struct FCameraNodeShakeDelta
+{
+	/** Position shake, which will be scaled and applied as a delta in the desired space. */
+	FVector3d Location = FVector3d::ZeroVector;
+
+	/** Rotation shake, which will be scaled and applied as a delta in the desired space. */
+	FRotator3d Rotation = FRotator3d::ZeroRotator;
+
+	/** Field of view shake, which will be scaled and applied as a delta. */
+	float FieldOfView = 0.f;
+
+public:
+
+	/** Combines this shake delta with another, optionally scaled, shake delta. */
+	void Combine(const FCameraNodeShakeDelta& Other, float OtherScale = 1.f);
 };
 
 /**
@@ -51,8 +71,23 @@ struct FCameraNodeShakeResult
 	/** The result that should be shaken. */
 	FCameraNodeEvaluationResult& ShakenResult;
 
+	/** 
+	 * Delta values that should be applied to the shaken result.
+	 *
+	 * These are simpler to set and combine inside camera shake nodes, since they will be applied
+	 * with the correct scale and in the correct space at the end. By comparison, writing to the
+	 * ShakenResult field directly requires taking into account all the parameters first (see 
+	 * FCameraNodeShakeParams).
+	 */
+	FCameraNodeShakeDelta ShakeDelta;
+
 	/** The time left in this shake, if applicable. Negative values indicate an infinite shake. */
 	float ShakeTimeLeft = 0.f;
+
+public:
+
+	/** Apply the delta values to the shaken result. */
+	void ApplyDelta(const FCameraNodeShakeParams& Params);
 };
 
 /**
diff --git a/Engine/Plugins/Cameras/GameplayCameras/Source/GameplayCameras/Public/GameFramework/GameplayCamerasPlayerCameraManager.h b/Engine/Plugins/Cameras/GameplayCameras/Source/GameplayCameras/Public/GameFramework/GameplayCamerasPlayerCameraManager.h
--- a/Engine/Plugins/Cameras/GameplayCameras/Source/GameplayCameras/Public/GameFramework/GameplayCamerasPlayerCameraManager.h
+++ b/Engine/Plugins/Cameras/GameplayCameras/Source/GameplayCameras/Public/GameFramework/GameplayCamerasPlayerCameraManager.h
@@ -4,6 +4,7 @@
 
 #include "Camera/PlayerCameraManager.h"
 #include "Core/CameraEvaluationContext.h"
+#include "Core/CameraShakeInstanceID.h"
 #include "GameFramework/IGameplayCameraSystemHost.h"
 
 #include "GameplayCamerasPlayerCameraManager.generated.h"
@@ -108,6 +109,24 @@ public:
 	UFUNCTION(BlueprintCallable, Category="Camera")
 	void StopCameraModifierRig(FCameraRigInstanceID InstanceID, bool bImmediately = false);
 
+public:
+
+	/** Starts a new camera shake. */
+	UFUNCTION(BlueprintCallable, Category="Camera")
+	FCameraShakeInstanceID StartCameraShakeAsset(
+			const UCameraShakeAsset* CameraShake,
+			float ShakeScale = 1.f,
+			ECameraShakePlaySpace PlaySpace = ECameraShakePlaySpace::CameraLocal,
+			FRotator UserPlaySpaceRotation = FRotator::ZeroRotator);
+
+	/** Checks if a camera shake is running. */
+	UFUNCTION(BlueprintCallable, Category="Camera")
+	bool IsCameraShakeAssetPlaying(FCameraShakeInstanceID InInstanceID) const;
+
+	/** Stops a running camera shake. */
+	UFUNCTION(BlueprintCallable, Category="Camera")
+	bool StopCameraShakeAsset(FCameraShakeInstanceID InInstanceID, bool bImmediately = false);
+
 public:
 
 	// APlayerCameraManager interface.
diff --git a/Engine/Plugins/Cameras/GameplayCameras/Source/GameplayCameras/Public/Services/CameraShakeService.h b/Engine/Plugins/Cameras/GameplayCameras/Source/GameplayCameras/Public/Services/CameraShakeService.h
--- a/Engine/Plugins/Cameras/GameplayCameras/Source/GameplayCameras/Public/Services/CameraShakeService.h
+++ b/Engine/Plugins/Cameras/GameplayCameras/Source/GameplayCameras/Public/Services/CameraShakeService.h
@@ -5,6 +5,7 @@
 #include "Camera/CameraTypes.h"
 #include "Core/CameraNode.h"
 #include "Core/CameraEvaluationService.h"
+#include "Core/CameraShakeInstanceID.h"
 
 #include "CameraShakeService.generated.h"
 
@@ -29,7 +30,7 @@ struct FStartCameraShakeParams
 	/** The play space to modify the result by */
 	ECameraShakePlaySpace PlaySpace = ECameraShakePlaySpace::CameraLocal;
 	/** The custom space to use for the shake. Only used when PlaySpace is UserDefined. */
-	FMatrix UserPlaySpaceMatrix;
+	FRotator UserPlaySpaceRotation;
 };
 
 /**
@@ -42,7 +43,13 @@ class FCameraShakeService : public FCameraEvaluationService
 public:
 
 	/** Starts a new camera shake. */
-	GAMEPLAYCAMERAS_API void StartCameraShake(const FStartCameraShakeParams& Params);
+	GAMEPLAYCAMERAS_API FCameraShakeInstanceID StartCameraShake(const FStartCameraShakeParams& Params);
+
+	/** Checks if a camera shake is running. */
+	GAMEPLAYCAMERAS_API bool IsCameraShakePlaying(FCameraShakeInstanceID InInstanceID) const;
+
+	/** Stops a running camera shake. */
+	GAMEPLAYCAMERAS_API bool StopCameraShake(FCameraShakeInstanceID InInstanceID, bool bImmediately = false);
 
 	/**
 	 * Requests that a given camera shake runs. Multiple requests for the same
@@ -63,10 +70,16 @@ private:
 
 private:
 
+	/** The system running this service. */
 	FCameraSystemEvaluator* Evaluator = nullptr;
 
+	/** The evaluation context used for running the shakes. */
 	TSharedPtr<FCameraEvaluationContext> ShakeContext;
+
+	/** The camera rig used to create the shake container. */
 	TObjectPtr<UCameraRigAsset> ShakeContainerRig = nullptr;
+
+	/** The shake container evaluator. */
 	FCameraShakeServiceCameraNodeEvaluator* ShakeEvaluator = nullptr;
 };
 
