Cameras: add bone/socket options to the attach to player pawn node

diff --git a/Engine/Plugins/Cameras/GameplayCameras/Source/GameplayCameras/Private/Nodes/Attach/AttachToPlayerPawnCameraNode.cpp b/Engine/Plugins/Cameras/GameplayCameras/Source/GameplayCameras/Private/Nodes/Attach/AttachToPlayerPawnCameraNode.cpp
--- a/Engine/Plugins/Cameras/GameplayCameras/Source/GameplayCameras/Private/Nodes/Attach/AttachToPlayerPawnCameraNode.cpp
+++ b/Engine/Plugins/Cameras/GameplayCameras/Source/GameplayCameras/Private/Nodes/Attach/AttachToPlayerPawnCameraNode.cpp
@@ -2,6 +2,7 @@
 
 #include "Nodes/Attach/AttachToPlayerPawnCameraNode.h"
 
+#include "Components/SkeletalMeshComponent.h"
 #include "Core/CameraEvaluationContext.h"
 #include "Core/CameraParameterReader.h"
 #include "Core/CameraSystemEvaluator.h"
@@ -25,6 +26,9 @@ protected:
 
 private:
 
+	USkeletalMeshComponent* CachedSkeletalMeshComponent = nullptr;
+	FName CachedBoneName;
+
 	TCameraParameterReader<bool> AttachToLocationReader;
 	TCameraParameterReader<bool> AttachToRotationReader;
 
@@ -47,6 +51,31 @@ void FAttachToPlayerPawnCameraNodeEvaluator::OnInitialize(const FCameraNodeEvalu
 		UE_LOG(LogCameraSystem, Error, 
 				TEXT("Can't run AttatchToPlayerPawn camera node because no player controller was found on the context."));
 		bHasValidPlayerController = false;
+		return;
+	}
+
+	APawn* Pawn = PlayerController->GetPawnOrSpectator();
+	if (!Pawn)
+	{
+		UE_LOG(LogCameraSystem, Error, 
+				TEXT("Can't run AttatchToPlayerPawn camera node because the player controller has no pawn."));
+		bHasValidPlayerController = false;
+		return;
+	}
+
+	if (!AttachNode->SocketName.IsNone() || !AttachNode->BoneName.IsNone())
+	{
+		CachedSkeletalMeshComponent = Pawn->FindComponentByClass<USkeletalMeshComponent>();
+	}
+
+	CachedBoneName = NAME_None;
+	if (CachedSkeletalMeshComponent)
+	{
+		CachedBoneName = AttachNode->BoneName;
+		if (!AttachNode->SocketName.IsNone())
+		{
+			CachedBoneName = CachedSkeletalMeshComponent->GetSocketBoneName(AttachNode->SocketName);
+		}
 	}
 }
 
@@ -80,17 +109,21 @@ void FAttachToPlayerPawnCameraNodeEvaluator::OnRun(const FCameraNodeEvaluationPa
 
 	const bool bAttachToLocation = AttachToLocationReader.Get(OutResult.VariableTable);
 	const bool bAttachToRotation = AttachToRotationReader.Get(OutResult.VariableTable);
+
+	FTransform3d AttachTransform = Pawn->GetActorTransform();
+	if (CachedSkeletalMeshComponent && !CachedBoneName.IsNone())
+	{
+		AttachTransform = CachedSkeletalMeshComponent->GetBoneTransform(CachedBoneName);
+	}
 	
 	if (bAttachToLocation)
 	{
-		const FVector3d PawnLocation = Pawn->GetActorLocation();
-		OutResult.CameraPose.SetLocation(PawnLocation);
+		OutResult.CameraPose.SetLocation(AttachTransform.GetLocation());
 	}
 
 	if (bAttachToRotation)
 	{
-		const FRotator3d PawnRotation = Pawn->GetActorRotation();
-		OutResult.CameraPose.SetRotation(PawnRotation);
+		OutResult.CameraPose.SetRotation(AttachTransform.GetRotation().Rotator());
 	}
 }
 
diff --git a/Engine/Plugins/Cameras/GameplayCameras/Source/GameplayCameras/Public/Nodes/Attach/AttachToPlayerPawnCameraNode.h b/Engine/Plugins/Cameras/GameplayCameras/Source/GameplayCameras/Public/Nodes/Attach/AttachToPlayerPawnCameraNode.h
--- a/Engine/Plugins/Cameras/GameplayCameras/Source/GameplayCameras/Public/Nodes/Attach/AttachToPlayerPawnCameraNode.h
+++ b/Engine/Plugins/Cameras/GameplayCameras/Source/GameplayCameras/Public/Nodes/Attach/AttachToPlayerPawnCameraNode.h
@@ -29,5 +29,13 @@ public:
 	/** Whether to align the camera rotation to the pawn's orientation. */
 	UPROPERTY(EditAnywhere, Category=Common)
 	FBooleanCameraParameter AttachToRotation = false;
+
+	/** An optional socket to attach to on the actor. */
+	UPROPERTY(BlueprintReadWrite, EditAnywhere, Interp, Category=Common)
+	FName SocketName;
+
+	/** An optional bone to attach to on the actor. */
+	UPROPERTY(BlueprintReadWrite, EditAnywhere, Interp, Category=Common)
+	FName BoneName;
 };
 
