Cameras: lazily build camera rigs when they're first activated

We normally build camera rigs and assets referenced by camera components when PIE starts, but we may be missing some rigs if they are indirectly referenced or not referenced at all by these components. This fixed these situations.

diff --git a/Engine/Plugins/Cameras/GameplayCameras/Source/GameplayCameras/Private/Core/BlendStackCameraNode.cpp b/Engine/Plugins/Cameras/GameplayCameras/Source/GameplayCameras/Private/Core/BlendStackCameraNode.cpp
--- a/Engine/Plugins/Cameras/GameplayCameras/Source/GameplayCameras/Private/Core/BlendStackCameraNode.cpp
+++ b/Engine/Plugins/Cameras/GameplayCameras/Source/GameplayCameras/Private/Core/BlendStackCameraNode.cpp
@@ -2,6 +2,8 @@
 
 #include "Core/BlendStackCameraNode.h"
 
+#include "Build/CameraRigAssetBuilder.h"
+#include "Build/CameraBuildLog.h"
 #include "Core/BlendStackCameraRigEvent.h"
 #include "Core/BlendStackRootCameraNode.h"
 #include "Core/CameraAsset.h"
@@ -20,6 +22,8 @@
 #include "Debug/CameraPoseLocationTrailDebugBlock.h"
 #include "Debug/ContextDataTableDebugBlock.h"
 #include "Debug/VariableTableDebugBlock.h"
+#include "Engine/World.h"
+#include "GameplayCamerasSettings.h"
 #include "HAL/IConsoleManager.h"
 #include "Helpers/CameraObjectInterfaceParameterOverrideHelper.h"
 #include "IGameplayCamerasLiveEditManager.h"
@@ -536,6 +540,31 @@ void FBlendStackCameraNodeEvaluator::RemoveListenedPackages(TSharedPtr<IGameplay
 	Entry.ListenedPackages.Reset();
 }
 
+void FBlendStackCameraNodeEvaluator::BuildCameraRigIfNeeded(const UCameraRigAsset* InCameraRig)
+{
+	if (!OwningEvaluator || !OwningEvaluator->GetOwner())
+	{
+		return;
+	}
+
+	if (BuiltCameraRigs.Contains(InCameraRig))
+	{
+		return;
+	}
+
+	UWorld* World = OwningEvaluator->GetOwner()->GetWorld();
+	const UGameplayCamerasSettings* Settings = GetDefault<UGameplayCamerasSettings>();
+	if (Settings->bAutoBuildInPIE && World && World->WorldType == EWorldType::PIE)
+	{
+		// Auto-build the camera rig the first time we see it.
+		FCameraBuildLog BuildLog;
+		FCameraRigAssetBuilder Builder(BuildLog);
+		Builder.BuildCameraRig(const_cast<UCameraRigAsset*>(InCameraRig));
+	}
+
+	BuiltCameraRigs.Add(InCameraRig);
+}
+
 #endif  // WITH_EDITOR
 
 void FBlendStackCameraNodeEvaluator::BroadcastCameraRigEvent(EBlendStackCameraRigEventType EventType, const FCameraRigEntry& Entry, const UCameraRigTransition* Transition) const
diff --git a/Engine/Plugins/Cameras/GameplayCameras/Source/GameplayCameras/Private/Core/PersistentBlendStackCameraNode.cpp b/Engine/Plugins/Cameras/GameplayCameras/Source/GameplayCameras/Private/Core/PersistentBlendStackCameraNode.cpp
--- a/Engine/Plugins/Cameras/GameplayCameras/Source/GameplayCameras/Private/Core/PersistentBlendStackCameraNode.cpp
+++ b/Engine/Plugins/Cameras/GameplayCameras/Source/GameplayCameras/Private/Core/PersistentBlendStackCameraNode.cpp
@@ -19,6 +19,10 @@ UE_DEFINE_CAMERA_NODE_EVALUATOR(FPersistentBlendStackCameraNodeEvaluator)
 
 FBlendStackEntryID FPersistentBlendStackCameraNodeEvaluator::Insert(const FBlendStackCameraInsertParams& Params)
 {
+#if WITH_EDITOR
+	BuildCameraRigIfNeeded(Params.CameraRig);
+#endif  // WITH_EDITOR
+
 	// See if we already have this camera rig and evaluation context in the stack.
 	if (!Params.bForceInsert)
 	{
diff --git a/Engine/Plugins/Cameras/GameplayCameras/Source/GameplayCameras/Private/Core/TransientBlendStackCameraNode.cpp b/Engine/Plugins/Cameras/GameplayCameras/Source/GameplayCameras/Private/Core/TransientBlendStackCameraNode.cpp
--- a/Engine/Plugins/Cameras/GameplayCameras/Source/GameplayCameras/Private/Core/TransientBlendStackCameraNode.cpp
+++ b/Engine/Plugins/Cameras/GameplayCameras/Source/GameplayCameras/Private/Core/TransientBlendStackCameraNode.cpp
@@ -23,6 +23,10 @@ UE_DEFINE_CAMERA_NODE_EVALUATOR(FTransientBlendStackCameraNodeEvaluator)
 
 FBlendStackEntryID FTransientBlendStackCameraNodeEvaluator::Push(const FBlendStackCameraPushParams& Params)
 {
+#if WITH_EDITOR
+	BuildCameraRigIfNeeded(Params.CameraRig);
+#endif  // WITH_EDITOR
+
 	bool bSearchedForTransition = false;
 	const UCameraRigTransition* Transition = nullptr;
 
diff --git a/Engine/Plugins/Cameras/GameplayCameras/Source/GameplayCameras/Public/Core/BlendStackCameraNode.h b/Engine/Plugins/Cameras/GameplayCameras/Source/GameplayCameras/Public/Core/BlendStackCameraNode.h
--- a/Engine/Plugins/Cameras/GameplayCameras/Source/GameplayCameras/Public/Core/BlendStackCameraNode.h
+++ b/Engine/Plugins/Cameras/GameplayCameras/Source/GameplayCameras/Public/Core/BlendStackCameraNode.h
@@ -167,6 +167,8 @@ protected:
 	void AddPackageListeners(FCameraRigEntry& Entry);
 	void RemoveListenedPackages(FCameraRigEntry& Entry);
 	void RemoveListenedPackages(TSharedPtr<IGameplayCamerasLiveEditManager> LiveEditManager, FCameraRigEntry& Entry);
+
+	void BuildCameraRigIfNeeded(const UCameraRigAsset* InCameraRig);
 #endif
 
 protected:
@@ -253,6 +255,7 @@ protected:
 
 #if WITH_EDITOR
 	TMap<TWeakObjectPtr<const UPackage>, int32> AllListenedPackages;
+	TSet<TWeakObjectPtr<const UCameraRigAsset>> BuiltCameraRigs;
 #endif  // WITH_EDITOR
 
 #if UE_GAMEPLAY_CAMERAS_DEBUG
