Cameras: remove unnecessary asserts when accessing context data entries

diff --git a/Engine/Plugins/Cameras/GameplayCameras/Source/GameplayCameras/Private/Core/CameraContextDataTable.cpp b/Engine/Plugins/Cameras/GameplayCameras/Source/GameplayCameras/Private/Core/CameraContextDataTable.cpp
--- a/Engine/Plugins/Cameras/GameplayCameras/Source/GameplayCameras/Private/Core/CameraContextDataTable.cpp
+++ b/Engine/Plugins/Cameras/GameplayCameras/Source/GameplayCameras/Private/Core/CameraContextDataTable.cpp
@@ -598,22 +598,14 @@ void FCameraContextDataTable::SetInstancedStructArrayData(FCameraContextDataID I
 
 const FCameraContextDataTable::FEntry* FCameraContextDataTable::FindEntry(FCameraContextDataID InID) const
 {
-	const uint32 EntryIndex = EntryLookup.FindRef(InID, INDEX_NONE);
-	if (ensure(Entries.IsValidIndex(EntryIndex)))
-	{
-		return &Entries[EntryIndex];
-	}
-	return nullptr;
+	const int32* IndexPtr = EntryLookup.Find(InID);
+	return IndexPtr ? &Entries[*IndexPtr] : nullptr;
 }
 
 FCameraContextDataTable::FEntry* FCameraContextDataTable::FindEntry(FCameraContextDataID InID)
 {
-	const uint32 EntryIndex = EntryLookup.FindRef(InID, INDEX_NONE);
-	if (ensure(Entries.IsValidIndex(EntryIndex)))
-	{
-		return &Entries[EntryIndex];
-	}
-	return nullptr;
+	const int32* IndexPtr = EntryLookup.Find(InID);
+	return IndexPtr ? &Entries[*IndexPtr] : nullptr;
 }
 
 const uint8* FCameraContextDataTable::GetData(
