Cameras: fix incorrect parameter/pin types on the Blueprint camera node for object and class pointers

diff --git a/Engine/Plugins/Cameras/GameplayCameras/Source/GameplayCameras/Private/Nodes/Utility/BlueprintCameraNode.cpp b/Engine/Plugins/Cameras/GameplayCameras/Source/GameplayCameras/Private/Nodes/Utility/BlueprintCameraNode.cpp
--- a/Engine/Plugins/Cameras/GameplayCameras/Source/GameplayCameras/Private/Nodes/Utility/BlueprintCameraNode.cpp
+++ b/Engine/Plugins/Cameras/GameplayCameras/Source/GameplayCameras/Private/Nodes/Utility/BlueprintCameraNode.cpp
@@ -516,15 +516,17 @@ void UBlueprintCameraNode::RebuildOverrides()
 			DataPropertyType = ECameraContextDataType::Enum;
 			DataPropertyTypeObject = EnumProperty->GetEnum();
 		}
-		else if (FObjectProperty* ObjectProperty = CastField<FObjectProperty>(Property))
+		else if (FClassProperty* ClassProperty = CastField<FClassProperty>(Property))
 		{
 			bIsDataProperty = true;
-			DataPropertyType = ECameraContextDataType::Object;
+			DataPropertyType = ECameraContextDataType::Class;
+			DataPropertyTypeObject = ClassProperty->MetaClass;
 		}
-		else if (FClassProperty* ClassProperty = CastField<FClassProperty>(Property))
+		else if (FObjectProperty* ObjectProperty = CastField<FObjectProperty>(Property))
 		{
 			bIsDataProperty = true;
-			DataPropertyType = ECameraContextDataType::Class;
+			DataPropertyType = ECameraContextDataType::Object;
+			DataPropertyTypeObject = ObjectProperty->PropertyClass;
 		}
 
 		if (bIsBlendableProperty)
