Cameras: fix pin types and colors for some data parameters

diff --git a/Engine/Plugins/Cameras/GameplayCameras/Source/GameplayCamerasEditor/Private/Editors/CameraNodeGraphNode.cpp b/Engine/Plugins/Cameras/GameplayCameras/Source/GameplayCamerasEditor/Private/Editors/CameraNodeGraphNode.cpp
--- a/Engine/Plugins/Cameras/GameplayCameras/Source/GameplayCamerasEditor/Private/Editors/CameraNodeGraphNode.cpp
+++ b/Engine/Plugins/Cameras/GameplayCameras/Source/GameplayCamerasEditor/Private/Editors/CameraNodeGraphNode.cpp
@@ -125,16 +125,16 @@ void UCameraNodeGraphNode::AllocateDefaultPins()
 				PinType.PinSubCategoryObject = StructProperty->Struct;
 				PinToolTip = StructProperty->Struct->GetDisplayNameText().ToString();
 			}
+			else if (FClassProperty* ClassProperty = CastField<FClassProperty>(DataProperty))
+			{
+				PinType.PinSubCategory = UEnum::GetValueAsName(ECameraContextDataType::Class);
+				PinType.PinSubCategoryObject = ClassProperty->MetaClass;
+			}
 			else if (FObjectProperty* ObjectProperty = CastField<FObjectProperty>(DataProperty))
 			{
 				PinType.PinSubCategory = UEnum::GetValueAsName(ECameraContextDataType::Object);
 				PinType.PinSubCategoryObject = ObjectProperty->PropertyClass;
 			}
-			else if (FClassProperty* ClassProperty = CastField<FClassProperty>(DataProperty))
-			{
-				PinType.PinSubCategory = UEnum::GetValueAsName(ECameraContextDataType::Class);
-				PinType.PinSubCategoryObject = ClassProperty->PropertyClass;
-			}
 			else
 			{
 				bGotValidDataProperty = false;
diff --git a/Engine/Plugins/Cameras/GameplayCameras/Source/GameplayCamerasEditor/Private/Editors/CameraNodeGraphPinColors.cpp b/Engine/Plugins/Cameras/GameplayCameras/Source/GameplayCamerasEditor/Private/Editors/CameraNodeGraphPinColors.cpp
--- a/Engine/Plugins/Cameras/GameplayCameras/Source/GameplayCamerasEditor/Private/Editors/CameraNodeGraphPinColors.cpp
+++ b/Engine/Plugins/Cameras/GameplayCameras/Source/GameplayCamerasEditor/Private/Editors/CameraNodeGraphPinColors.cpp
@@ -2,6 +2,7 @@
 
 #include "Editors/CameraNodeGraphPinColors.h"
 
+#include "Core/CameraContextDataTableFwd.h"
 #include "Core/CameraParameters.h"
 #include "Core/CameraVariableReferences.h"
 #include "Core/CameraVariableTableFwd.h"
@@ -64,7 +65,12 @@ void FCameraNodeGraphPinColors::Initialize()
 	PinColors.Add(FTransform3dCameraVariableReference::StaticStruct()->GetFName(), Settings->TransformPinTypeColor);
 
 	DefaultPinColor = Settings->DefaultPinTypeColor;
+	NamePinColor = Settings->NamePinTypeColor;
+	StringPinColor = Settings->StringPinTypeColor;
+	EnumPinColor = Settings->Int64PinTypeColor;
 	StructPinColor = Settings->StructPinTypeColor;
+	ObjectPinColor = Settings->ObjectPinTypeColor;
+	ClassPinColor = Settings->ClassPinTypeColor;
 }
 
 FLinearColor FCameraNodeGraphPinColors::GetPinColor(const FName& TypeName) const
@@ -72,9 +78,27 @@ FLinearColor FCameraNodeGraphPinColors::GetPinColor(const FName& TypeName) const
 	return PinColors.FindRef(TypeName, DefaultPinColor);
 }
 
-FLinearColor FCameraNodeGraphPinColors::GetStructPinColor() const
+FLinearColor FCameraNodeGraphPinColors::GetContextDataPinColor(const FName& DataTypeName) const
 {
-	return StructPinColor;
+	const UEnum* DataTypeEnum = StaticEnum<ECameraContextDataType>();
+	const ECameraContextDataType DataType = (ECameraContextDataType)DataTypeEnum->GetValueByName(DataTypeName);
+	switch (DataType)
+	{
+		case ECameraContextDataType::Name:
+			return NamePinColor;
+		case ECameraContextDataType::String:
+			return StringPinColor;
+		case ECameraContextDataType::Enum:
+			return EnumPinColor;
+		case ECameraContextDataType::Struct:
+			return StructPinColor;
+		case ECameraContextDataType::Object:
+			return ObjectPinColor;
+		case ECameraContextDataType::Class:
+			return ClassPinColor;
+		default:
+			return DefaultPinColor;
+	}
 }
 
 }  // namespace UE::Cameras
diff --git a/Engine/Plugins/Cameras/GameplayCameras/Source/GameplayCamerasEditor/Private/Editors/CameraNodeGraphPinColors.h b/Engine/Plugins/Cameras/GameplayCameras/Source/GameplayCamerasEditor/Private/Editors/CameraNodeGraphPinColors.h
--- a/Engine/Plugins/Cameras/GameplayCameras/Source/GameplayCamerasEditor/Private/Editors/CameraNodeGraphPinColors.h
+++ b/Engine/Plugins/Cameras/GameplayCameras/Source/GameplayCamerasEditor/Private/Editors/CameraNodeGraphPinColors.h
@@ -14,13 +14,18 @@ struct FCameraNodeGraphPinColors
 	void Initialize();
 
 	FLinearColor GetPinColor(const FName& TypeName) const;
-	FLinearColor GetStructPinColor() const;
+	FLinearColor GetContextDataPinColor(const FName& DataTypeName) const;
 
 private:
 
 	TMap<FName, FLinearColor> PinColors;
 	FLinearColor DefaultPinColor;
+	FLinearColor NamePinColor;
+	FLinearColor StringPinColor;
+	FLinearColor EnumPinColor;
 	FLinearColor StructPinColor;
+	FLinearColor ObjectPinColor;
+	FLinearColor ClassPinColor;
 };
 
 }  // namespace UE::Cameras
diff --git a/Engine/Plugins/Cameras/GameplayCameras/Source/GameplayCamerasEditor/Private/Editors/CameraNodeGraphSchema.cpp b/Engine/Plugins/Cameras/GameplayCameras/Source/GameplayCamerasEditor/Private/Editors/CameraNodeGraphSchema.cpp
--- a/Engine/Plugins/Cameras/GameplayCameras/Source/GameplayCamerasEditor/Private/Editors/CameraNodeGraphSchema.cpp
+++ b/Engine/Plugins/Cameras/GameplayCameras/Source/GameplayCamerasEditor/Private/Editors/CameraNodeGraphSchema.cpp
@@ -464,7 +464,8 @@ FLinearColor UCameraNodeGraphSchema::GetPinTypeColor(const FEdGraphPinType& PinT
 	}
 	if (PinType.PinCategory == PC_CameraContextData)
 	{
-		return PinColors.GetStructPinColor();
+		const FName DataTypeName = PinType.PinSubCategory;
+		return PinColors.GetContextDataPinColor(DataTypeName);
 	}
 
 	return UObjectTreeGraphSchema::GetPinTypeColor(PinType);
