Cameras: change how auto-build for PIE works to prevent packages from being modified at the wrong time

Building a camera asset may change its data and dirty its package, but the previous method was doing when we were already inside the PIE world. This change uses the same method as Blueprints, StateTrees, and so on, by using the BeginPIE delegate along with typed object iterators to gather objects to process just before PIE starts.

diff --git a/Engine/Plugins/Cameras/GameplayCameras/Source/GameplayCameras/Private/Build/CameraBuildLog.cpp b/Engine/Plugins/Cameras/GameplayCameras/Source/GameplayCameras/Private/Build/CameraBuildLog.cpp
--- a/Engine/Plugins/Cameras/GameplayCameras/Source/GameplayCameras/Private/Build/CameraBuildLog.cpp
+++ b/Engine/Plugins/Cameras/GameplayCameras/Source/GameplayCameras/Private/Build/CameraBuildLog.cpp
@@ -92,5 +92,12 @@ void FCameraBuildLog::AddMessage(EMessageSeverity::Type InSeverity, const UObjec
 	}
 }
 
+void FCameraBuildLog::ResetMessages()
+{
+	Messages.Reset();
+	bHasWarnings = false;
+	bHasErrors = false;
+}
+
 }  // namespace UE::Cameras
 
diff --git a/Engine/Plugins/Cameras/GameplayCameras/Source/GameplayCameras/Private/Core/BlendStackCameraNode.cpp b/Engine/Plugins/Cameras/GameplayCameras/Source/GameplayCameras/Private/Core/BlendStackCameraNode.cpp
--- a/Engine/Plugins/Cameras/GameplayCameras/Source/GameplayCameras/Private/Core/BlendStackCameraNode.cpp
+++ b/Engine/Plugins/Cameras/GameplayCameras/Source/GameplayCameras/Private/Core/BlendStackCameraNode.cpp
@@ -540,31 +540,6 @@ void FBlendStackCameraNodeEvaluator::RemoveListenedPackages(TSharedPtr<IGameplay
 	Entry.ListenedPackages.Reset();
 }
 
-void FBlendStackCameraNodeEvaluator::BuildCameraRigIfNeeded(const UCameraRigAsset* InCameraRig)
-{
-	if (!OwningEvaluator || !OwningEvaluator->GetOwner())
-	{
-		return;
-	}
-
-	if (BuiltCameraRigs.Contains(InCameraRig))
-	{
-		return;
-	}
-
-	UWorld* World = OwningEvaluator->GetOwner()->GetWorld();
-	const UGameplayCamerasSettings* Settings = GetDefault<UGameplayCamerasSettings>();
-	if (Settings->bAutoBuildInPIE && World && World->WorldType == EWorldType::PIE)
-	{
-		// Auto-build the camera rig the first time we see it.
-		FCameraBuildLog BuildLog;
-		FCameraRigAssetBuilder Builder(BuildLog);
-		Builder.BuildCameraRig(const_cast<UCameraRigAsset*>(InCameraRig));
-	}
-
-	BuiltCameraRigs.Add(InCameraRig);
-}
-
 #endif  // WITH_EDITOR
 
 void FBlendStackCameraNodeEvaluator::BroadcastCameraRigEvent(EBlendStackCameraRigEventType EventType, const FCameraRigEntry& Entry, const UCameraRigTransition* Transition) const
diff --git a/Engine/Plugins/Cameras/GameplayCameras/Source/GameplayCameras/Private/Core/PersistentBlendStackCameraNode.cpp b/Engine/Plugins/Cameras/GameplayCameras/Source/GameplayCameras/Private/Core/PersistentBlendStackCameraNode.cpp
--- a/Engine/Plugins/Cameras/GameplayCameras/Source/GameplayCameras/Private/Core/PersistentBlendStackCameraNode.cpp
+++ b/Engine/Plugins/Cameras/GameplayCameras/Source/GameplayCameras/Private/Core/PersistentBlendStackCameraNode.cpp
@@ -19,10 +19,6 @@ UE_DEFINE_CAMERA_NODE_EVALUATOR(FPersistentBlendStackCameraNodeEvaluator)
 
 FBlendStackEntryID FPersistentBlendStackCameraNodeEvaluator::Insert(const FBlendStackCameraInsertParams& Params)
 {
-#if WITH_EDITOR
-	BuildCameraRigIfNeeded(Params.CameraRig);
-#endif  // WITH_EDITOR
-
 	// See if we already have this camera rig and evaluation context in the stack.
 	if (!Params.bForceInsert)
 	{
diff --git a/Engine/Plugins/Cameras/GameplayCameras/Source/GameplayCameras/Private/Core/TransientBlendStackCameraNode.cpp b/Engine/Plugins/Cameras/GameplayCameras/Source/GameplayCameras/Private/Core/TransientBlendStackCameraNode.cpp
--- a/Engine/Plugins/Cameras/GameplayCameras/Source/GameplayCameras/Private/Core/TransientBlendStackCameraNode.cpp
+++ b/Engine/Plugins/Cameras/GameplayCameras/Source/GameplayCameras/Private/Core/TransientBlendStackCameraNode.cpp
@@ -23,10 +23,6 @@ UE_DEFINE_CAMERA_NODE_EVALUATOR(FTransientBlendStackCameraNodeEvaluator)
 
 FBlendStackEntryID FTransientBlendStackCameraNodeEvaluator::Push(const FBlendStackCameraPushParams& Params)
 {
-#if WITH_EDITOR
-	BuildCameraRigIfNeeded(Params.CameraRig);
-#endif  // WITH_EDITOR
-
 	bool bSearchedForTransition = false;
 	const UCameraRigTransition* Transition = nullptr;
 
diff --git a/Engine/Plugins/Cameras/GameplayCameras/Source/GameplayCameras/Private/GameFramework/GameplayCameraComponentBase.cpp b/Engine/Plugins/Cameras/GameplayCameras/Source/GameplayCameras/Private/GameFramework/GameplayCameraComponentBase.cpp
--- a/Engine/Plugins/Cameras/GameplayCameras/Source/GameplayCameras/Private/GameFramework/GameplayCameraComponentBase.cpp
+++ b/Engine/Plugins/Cameras/GameplayCameras/Source/GameplayCameras/Private/GameFramework/GameplayCameraComponentBase.cpp
@@ -2,8 +2,6 @@
 
 #include "GameFramework/GameplayCameraComponentBase.h"
 
-#include "Build/CameraAssetBuilder.h"
-#include "Build/CameraBuildLog.h"
 #include "CineCameraComponent.h"
 #include "Core/CameraAsset.h"
 #include "Core/CameraEvaluationContextStack.h"
@@ -15,7 +13,6 @@
 #include "Engine/Canvas.h"
 #include "Engine/EngineTypes.h"
 #include "Engine/World.h"
-#include "GameplayCamerasSettings.h"
 #include "IGameplayCamerasLiveEditManager.h"
 #include "IGameplayCamerasModule.h"
 #include "Kismet/GameplayStatics.h"
@@ -337,23 +334,6 @@ void UGameplayCameraComponentBase::CreateCameraEvaluationContext(APlayerControll
 	{
 		UCameraAsset* CameraAsset = GetCameraAsset();
 
-#if WITH_EDITOR
-		if (CameraAsset && bIsFirstActivation)
-		{
-			UWorld* World = GetWorld();
-			const UGameplayCamerasSettings* Settings = GetDefault<UGameplayCamerasSettings>();
-			if (Settings->bAutoBuildInPIE && World && World->WorldType == EWorldType::PIE)
-			{
-				// Auto-build the camera asset on begin play to make sure we've got the latest user edits.
-				FCameraBuildLog BuildLog;
-				FCameraAssetBuilder Builder(BuildLog);
-				Builder.BuildCamera(CameraAsset);
-			}
-
-			bIsFirstActivation = false;
-		}
-#endif
-		
 		// If we have no camera asset specified, make a placeholder one and log a warning.
 		if (!CameraAsset)
 		{
diff --git a/Engine/Plugins/Cameras/GameplayCameras/Source/GameplayCameras/Public/Build/CameraBuildLog.h b/Engine/Plugins/Cameras/GameplayCameras/Source/GameplayCameras/Public/Build/CameraBuildLog.h
--- a/Engine/Plugins/Cameras/GameplayCameras/Source/GameplayCameras/Public/Build/CameraBuildLog.h
+++ b/Engine/Plugins/Cameras/GameplayCameras/Source/GameplayCameras/Public/Build/CameraBuildLog.h
@@ -56,6 +56,9 @@ public:
 	/** Gets the list of received messages so far. */
 	TArrayView<const FCameraBuildLogMessage> GetMessages() const { return Messages; }
 
+	/** Clears all messages. */
+	UE_API void ResetMessages();
+
 	/** Returns whether any warning has been logged. */
 	bool HasWarnings() const { return bHasWarnings; }
 	/** Returns whether any error has been logged. */
diff --git a/Engine/Plugins/Cameras/GameplayCameras/Source/GameplayCameras/Public/Core/BlendStackCameraNode.h b/Engine/Plugins/Cameras/GameplayCameras/Source/GameplayCameras/Public/Core/BlendStackCameraNode.h
--- a/Engine/Plugins/Cameras/GameplayCameras/Source/GameplayCameras/Public/Core/BlendStackCameraNode.h
+++ b/Engine/Plugins/Cameras/GameplayCameras/Source/GameplayCameras/Public/Core/BlendStackCameraNode.h
@@ -167,8 +167,6 @@ protected:
 	void AddPackageListeners(FCameraRigEntry& Entry);
 	void RemoveListenedPackages(FCameraRigEntry& Entry);
 	void RemoveListenedPackages(TSharedPtr<IGameplayCamerasLiveEditManager> LiveEditManager, FCameraRigEntry& Entry);
-
-	void BuildCameraRigIfNeeded(const UCameraRigAsset* InCameraRig);
 #endif
 
 protected:
diff --git a/Engine/Plugins/Cameras/GameplayCameras/Source/GameplayCameras/Public/GameFramework/GameplayCameraComponentBase.h b/Engine/Plugins/Cameras/GameplayCameras/Source/GameplayCameras/Public/GameFramework/GameplayCameraComponentBase.h
--- a/Engine/Plugins/Cameras/GameplayCameras/Source/GameplayCameras/Public/GameFramework/GameplayCameraComponentBase.h
+++ b/Engine/Plugins/Cameras/GameplayCameras/Source/GameplayCameras/Public/GameFramework/GameplayCameraComponentBase.h
@@ -270,9 +270,6 @@ private:
 	/** Whether this component is running in an editor world. */
 	bool bIsEditorWorld = false;
 
-	/** Whether this component has never been activated so far. */
-	bool bIsFirstActivation = true;
-
 	/** The show flag for camera system debug rendering. */
 	int32 CustomShowFlag = INDEX_NONE;
 
diff --git a/Engine/Plugins/Cameras/GameplayCameras/Source/GameplayCamerasEditor/Private/GameplayCamerasLiveEditManager.cpp b/Engine/Plugins/Cameras/GameplayCameras/Source/GameplayCamerasEditor/Private/GameplayCamerasLiveEditManager.cpp
--- a/Engine/Plugins/Cameras/GameplayCameras/Source/GameplayCamerasEditor/Private/GameplayCamerasLiveEditManager.cpp
+++ b/Engine/Plugins/Cameras/GameplayCameras/Source/GameplayCamerasEditor/Private/GameplayCamerasLiveEditManager.cpp
@@ -2,11 +2,22 @@
 
 #include "GameplayCamerasLiveEditManager.h"
 
+#include "Build/CameraAssetBuilder.h"
+#include "Build/CameraRigAssetBuilder.h"
+#include "Build/CameraShakeAssetBuilder.h"
+#include "Core/CameraAsset.h"
+#include "Core/CameraRigAsset.h"
+#include "Core/CameraShakeAsset.h"
+#include "Editor.h"
+#include "GameplayCamerasSettings.h"
+#include "IGameplayCamerasEditorModule.h"
 #include "IGameplayCamerasLiveEditListener.h"
 #include "GameplayCamerasEditorSettings.h"
-#include "Misc/CoreDelegates.h"
 #include "UObject/Package.h"
 #include "UObject/UObjectGlobals.h"
+#include "UObject/UObjectIterator.h"
+
+#define LOCTEXT_NAMESPACE "GameplayCamerasLiveEditManager"
 
 namespace UE::Cameras
 {
@@ -55,10 +66,12 @@ void RemoveListenerImpl(
 FGameplayCamerasLiveEditManager::FGameplayCamerasLiveEditManager()
 {
 	FCoreUObjectDelegates::GetPostGarbageCollect().AddRaw(this, &FGameplayCamerasLiveEditManager::OnPostGarbageCollection);
+	FEditorDelegates::BeginPIE.AddRaw(this, &FGameplayCamerasLiveEditManager::OnBeginPIE);
 }
 
 FGameplayCamerasLiveEditManager::~FGameplayCamerasLiveEditManager()
 {
+	FEditorDelegates::BeginPIE.RemoveAll(this);
 	FCoreUObjectDelegates::GetPostGarbageCollect().RemoveAll(this);
 }
 
@@ -152,5 +165,126 @@ void FGameplayCamerasLiveEditManager::RemoveGarbage()
 	}
 }
 
+namespace Internal
+{
+
+template<typename CameraObjectType>
+struct TCameraObjectBuilderTraits;
+
+template<>
+struct TCameraObjectBuilderTraits<UCameraAsset>
+{
+	static void Build(UCameraAsset* InObject, FCameraBuildLog& InBuildLog)
+	{
+		const bool bBuildReferencedAssets = false; // We are going to build rigs first.
+		FCameraAssetBuilder Builder(InBuildLog);
+		Builder.BuildCamera(InObject, bBuildReferencedAssets); 
+	}
+};
+
+template<>
+struct TCameraObjectBuilderTraits<UCameraRigAsset>
+{
+	static void Build(UCameraRigAsset* InObject, FCameraBuildLog& InBuildLog)
+	{
+		FCameraRigAssetBuilder Builder(InBuildLog);
+		Builder.BuildCameraRig(InObject);
+	}
+};
+
+template<>
+struct TCameraObjectBuilderTraits<UCameraShakeAsset>
+{
+	static void Build(UCameraShakeAsset* InObject, FCameraBuildLog& InBuildLog)
+	{
+		FCameraShakeAssetBuilder Builder(InBuildLog);
+		Builder.BuildCameraShake(InObject); 
+	}
+};
+
+template<typename CameraObjectType>
+struct TCameraObjectBuilderUtil
+{
+	static int32 Gather(TArray<CameraObjectType*>& OutCameraObjectsToBuild)
+	{
+		int32 TotalCameraObjects = 0;
+		for (TObjectIterator<CameraObjectType> It; It; ++It)
+		{
+			CameraObjectType* Obj(*It);
+			++TotalCameraObjects;
+
+			const ECameraBuildStatus BuildStatus = Obj->GetBuildStatus();
+			if (BuildStatus == ECameraBuildStatus::Clean || BuildStatus == ECameraBuildStatus::CleanWithWarnings)
+			{
+				continue;
+			}
+
+			OutCameraObjectsToBuild.Add(Obj);
+		}
+		return TotalCameraObjects;
+	}
+
+	static void Build(TArrayView<CameraObjectType*> InCameraObjectsToBuild)
+	{
+		FCameraBuildLog BuildLog;
+		BuildLog.SetForwardMessagesToLogging(true);
+
+		for (CameraObjectType* CameraObject : InCameraObjectsToBuild)
+		{
+			BuildLog.ResetMessages();
+			BuildLog.SetLoggingPrefix(CameraObject->GetName());
+
+			TCameraObjectBuilderTraits<CameraObjectType>::Build(CameraObject, BuildLog);
+		}
+	}
+};
+
+}  // namespace Internal
+
+void FGameplayCamerasLiveEditManager::OnBeginPIE(const bool bSimulate)
+{
+	using namespace Internal;
+
+	const UGameplayCamerasSettings* Settings = GetDefault<UGameplayCamerasSettings>();
+	if (!Settings->bAutoBuildInPIE)
+	{
+		return;
+	}
+
+	TArray<UCameraRigAsset*> CameraRigsToBuild;
+	const int32 NumCameraRigs = TCameraObjectBuilderUtil<UCameraRigAsset>::Gather(CameraRigsToBuild);
+
+	TArray<UCameraShakeAsset*> CameraShakesToBuild;
+	const int32 NumCameraShakes = TCameraObjectBuilderUtil<UCameraShakeAsset>::Gather(CameraShakesToBuild);
+
+	TArray<UCameraAsset*> CamerasToBuild;
+	const int32 NumCameras = TCameraObjectBuilderUtil<UCameraAsset>::Gather(CamerasToBuild);
+
+	const int32 NumCameraObjects = (NumCameras + NumCameraRigs + NumCameraShakes);
+	const int32 NumCameraObjectsToBuild = (CamerasToBuild.Num() + CameraRigsToBuild.Num() + CameraShakesToBuild.Num());
+	if (NumCameraObjectsToBuild > 0)
+	{
+		const double BuildStartTime = FPlatformTime::Seconds();
+
+		TCameraObjectBuilderUtil<UCameraRigAsset>::Build(CameraRigsToBuild);
+		TCameraObjectBuilderUtil<UCameraShakeAsset>::Build(CameraShakesToBuild);
+		TCameraObjectBuilderUtil<UCameraAsset>::Build(CamerasToBuild);
+
+		const double BuildEndTime = FPlatformTime::Seconds();
+		UE_LOG(LogCameraSystemEditor, Log, 
+				TEXT("Built %d/%d camera objects in %d ms"),
+				NumCameraObjectsToBuild, NumCameraObjects,
+				(int32)((BuildEndTime - BuildStartTime) * 1000));
+	}
+	else
+	{
+		UE_LOG(LogCameraSystemEditor, Log, 
+				TEXT("No camera objects needed building (inspected %d objects)"),
+				NumCameraObjects);
+	}
+}
+
 }  // namespace UE::Cameras
 
+#undef LOCTEXT_NAMESPACE
+
diff --git a/Engine/Plugins/Cameras/GameplayCameras/Source/GameplayCamerasEditor/Private/GameplayCamerasLiveEditManager.h b/Engine/Plugins/Cameras/GameplayCameras/Source/GameplayCamerasEditor/Private/GameplayCamerasLiveEditManager.h
--- a/Engine/Plugins/Cameras/GameplayCameras/Source/GameplayCamerasEditor/Private/GameplayCamerasLiveEditManager.h
+++ b/Engine/Plugins/Cameras/GameplayCameras/Source/GameplayCamerasEditor/Private/GameplayCamerasLiveEditManager.h
@@ -34,6 +34,7 @@ public:
 private:
 
 	void OnPostGarbageCollection();
+	void OnBeginPIE(const bool bSimulate);
 
 	void RemoveGarbage();
 
