Cameras: fix crash when opening the Blueprint palette

diff --git a/Engine/Plugins/Cameras/GameplayCameras/Source/GameplayCamerasUncookedOnly/Private/BlueprintGraph/K2Node_GetCameraRigParameter.cpp b/Engine/Plugins/Cameras/GameplayCameras/Source/GameplayCamerasUncookedOnly/Private/BlueprintGraph/K2Node_GetCameraRigParameter.cpp
--- a/Engine/Plugins/Cameras/GameplayCameras/Source/GameplayCamerasUncookedOnly/Private/BlueprintGraph/K2Node_GetCameraRigParameter.cpp
+++ b/Engine/Plugins/Cameras/GameplayCameras/Source/GameplayCamerasUncookedOnly/Private/BlueprintGraph/K2Node_GetCameraRigParameter.cpp
@@ -55,7 +55,7 @@ void UK2Node_GetCameraRigParameter::GetMenuActions(FBlueprintActionDatabaseRegis
 				[ParameterName, CameraRigAssetData](UEdGraphNode* NewNode, bool bIsTemplateNode)
 				{
 					UK2Node_GetCameraRigParameter* NewSetter = CastChecked<UK2Node_GetCameraRigParameter>(NewNode);
-					NewSetter->Initialize(CameraRigAssetData, ParameterName.ToString());
+					NewSetter->Initialize(CameraRigAssetData, ParameterName.ToString(), bIsTemplateNode);
 				});
 
 		ActionRegistrar.AddBlueprintAction(CameraRigAssetData, NodeSpawner);
diff --git a/Engine/Plugins/Cameras/GameplayCameras/Source/GameplayCamerasUncookedOnly/Private/BlueprintGraph/K2Node_GetCameraRigParameters.cpp b/Engine/Plugins/Cameras/GameplayCameras/Source/GameplayCamerasUncookedOnly/Private/BlueprintGraph/K2Node_GetCameraRigParameters.cpp
--- a/Engine/Plugins/Cameras/GameplayCameras/Source/GameplayCamerasUncookedOnly/Private/BlueprintGraph/K2Node_GetCameraRigParameters.cpp
+++ b/Engine/Plugins/Cameras/GameplayCameras/Source/GameplayCamerasUncookedOnly/Private/BlueprintGraph/K2Node_GetCameraRigParameters.cpp
@@ -49,7 +49,7 @@ void UK2Node_GetCameraRigParameters::GetMenuActions(FBlueprintActionDatabaseRegi
 			[CameraRigAssetData](UEdGraphNode* NewNode, bool bIsTemplateNode)
 			{
 				UK2Node_GetCameraRigParameters* NewGetter = CastChecked<UK2Node_GetCameraRigParameters>(NewNode);
-				NewGetter->Initialize(CameraRigAssetData);
+				NewGetter->Initialize(CameraRigAssetData, bIsTemplateNode);
 			});
 
 	ActionRegistrar.AddBlueprintAction(CameraRigAssetData, NodeSpawner);
diff --git a/Engine/Plugins/Cameras/GameplayCameras/Source/GameplayCamerasUncookedOnly/Private/BlueprintGraph/K2Node_MultiCameraRigParametersBase.cpp b/Engine/Plugins/Cameras/GameplayCameras/Source/GameplayCamerasUncookedOnly/Private/BlueprintGraph/K2Node_MultiCameraRigParametersBase.cpp
--- a/Engine/Plugins/Cameras/GameplayCameras/Source/GameplayCamerasUncookedOnly/Private/BlueprintGraph/K2Node_MultiCameraRigParametersBase.cpp
+++ b/Engine/Plugins/Cameras/GameplayCameras/Source/GameplayCamerasUncookedOnly/Private/BlueprintGraph/K2Node_MultiCameraRigParametersBase.cpp
@@ -21,9 +21,12 @@ UK2Node_MultiCameraRigParametersBase::UK2Node_MultiCameraRigParametersBase(const
 {
 }
 
-void UK2Node_MultiCameraRigParametersBase::Initialize(const FAssetData& UnloadedCameraRig)
+void UK2Node_MultiCameraRigParametersBase::Initialize(const FAssetData& UnloadedCameraRig, bool bIsTemplateNode)
 {
-	CameraRig = Cast<UCameraRigAsset>(UnloadedCameraRig.GetAsset());
+	if (!bIsTemplateNode)
+	{
+		CameraRig = Cast<UCameraRigAsset>(UnloadedCameraRig.GetAsset());
+	}
 }
 
 void UK2Node_MultiCameraRigParametersBase::ReallocatePinsDuringReconstruction(TArray<UEdGraphPin*>& OldPins)
diff --git a/Engine/Plugins/Cameras/GameplayCameras/Source/GameplayCamerasUncookedOnly/Private/BlueprintGraph/K2Node_MultiCameraRigParametersBase.h b/Engine/Plugins/Cameras/GameplayCameras/Source/GameplayCamerasUncookedOnly/Private/BlueprintGraph/K2Node_MultiCameraRigParametersBase.h
--- a/Engine/Plugins/Cameras/GameplayCameras/Source/GameplayCamerasUncookedOnly/Private/BlueprintGraph/K2Node_MultiCameraRigParametersBase.h
+++ b/Engine/Plugins/Cameras/GameplayCameras/Source/GameplayCamerasUncookedOnly/Private/BlueprintGraph/K2Node_MultiCameraRigParametersBase.h
@@ -20,7 +20,7 @@ public:
 
 	UK2Node_MultiCameraRigParametersBase(const FObjectInitializer& ObjectInit);
 
-	void Initialize(const FAssetData& UnloadedCameraRig);
+	void Initialize(const FAssetData& UnloadedCameraRig, bool bIsTemplateNode);
 
 public:
 
diff --git a/Engine/Plugins/Cameras/GameplayCameras/Source/GameplayCamerasUncookedOnly/Private/BlueprintGraph/K2Node_SetCameraRigParameter.cpp b/Engine/Plugins/Cameras/GameplayCameras/Source/GameplayCamerasUncookedOnly/Private/BlueprintGraph/K2Node_SetCameraRigParameter.cpp
--- a/Engine/Plugins/Cameras/GameplayCameras/Source/GameplayCamerasUncookedOnly/Private/BlueprintGraph/K2Node_SetCameraRigParameter.cpp
+++ b/Engine/Plugins/Cameras/GameplayCameras/Source/GameplayCamerasUncookedOnly/Private/BlueprintGraph/K2Node_SetCameraRigParameter.cpp
@@ -55,7 +55,7 @@ void UK2Node_SetCameraRigParameter::GetMenuActions(FBlueprintActionDatabaseRegis
 				[ParameterName, CameraRigAssetData](UEdGraphNode* NewNode, bool bIsTemplateNode)
 				{
 					UK2Node_SetCameraRigParameter* NewSetter = CastChecked<UK2Node_SetCameraRigParameter>(NewNode);
-					NewSetter->Initialize(CameraRigAssetData, ParameterName.ToString());
+					NewSetter->Initialize(CameraRigAssetData, ParameterName.ToString(), bIsTemplateNode);
 				});
 
 		ActionRegistrar.AddBlueprintAction(CameraRigAssetData, NodeSpawner);
diff --git a/Engine/Plugins/Cameras/GameplayCameras/Source/GameplayCamerasUncookedOnly/Private/BlueprintGraph/K2Node_SetCameraRigParameters.cpp b/Engine/Plugins/Cameras/GameplayCameras/Source/GameplayCamerasUncookedOnly/Private/BlueprintGraph/K2Node_SetCameraRigParameters.cpp
--- a/Engine/Plugins/Cameras/GameplayCameras/Source/GameplayCamerasUncookedOnly/Private/BlueprintGraph/K2Node_SetCameraRigParameters.cpp
+++ b/Engine/Plugins/Cameras/GameplayCameras/Source/GameplayCamerasUncookedOnly/Private/BlueprintGraph/K2Node_SetCameraRigParameters.cpp
@@ -87,7 +87,7 @@ void UK2Node_SetCameraRigParameters::GetMenuActions(FBlueprintActionDatabaseRegi
 			[CameraRigAssetData](UEdGraphNode* NewNode, bool bIsTemplateNode)
 			{
 				UK2Node_SetCameraRigParameters* NewSetter = CastChecked<UK2Node_SetCameraRigParameters>(NewNode);
-				NewSetter->Initialize(CameraRigAssetData);
+				NewSetter->Initialize(CameraRigAssetData, bIsTemplateNode);
 			});
 
 	ActionRegistrar.AddBlueprintAction(CameraRigAssetData, NodeSpawner);
diff --git a/Engine/Plugins/Cameras/GameplayCameras/Source/GameplayCamerasUncookedOnly/Private/BlueprintGraph/K2Node_SingleCameraRigParameterBase.cpp b/Engine/Plugins/Cameras/GameplayCameras/Source/GameplayCamerasUncookedOnly/Private/BlueprintGraph/K2Node_SingleCameraRigParameterBase.cpp
--- a/Engine/Plugins/Cameras/GameplayCameras/Source/GameplayCamerasUncookedOnly/Private/BlueprintGraph/K2Node_SingleCameraRigParameterBase.cpp
+++ b/Engine/Plugins/Cameras/GameplayCameras/Source/GameplayCamerasUncookedOnly/Private/BlueprintGraph/K2Node_SingleCameraRigParameterBase.cpp
@@ -24,20 +24,23 @@ UK2Node_SingleCameraRigParameterBase::UK2Node_SingleCameraRigParameterBase(const
 {
 }
 
-void UK2Node_SingleCameraRigParameterBase::Initialize(const FAssetData& UnloadedCameraRig, const FString& InCameraParameterName)
+void UK2Node_SingleCameraRigParameterBase::Initialize(const FAssetData& UnloadedCameraRig, const FString& InCameraParameterName, bool bIsTemplateNode)
 {
-	UCameraRigAsset* LoadedCameraRig = Cast<UCameraRigAsset>(UnloadedCameraRig.GetAsset());
-	if (ensure(LoadedCameraRig))
+	if (!bIsTemplateNode)
 	{
-		if (const UCameraObjectInterfaceBlendableParameter* BlendableParameter = LoadedCameraRig->Interface.FindBlendableParameterByName(InCameraParameterName))
+		UCameraRigAsset* LoadedCameraRig = Cast<UCameraRigAsset>(UnloadedCameraRig.GetAsset());
+		if (ensure(LoadedCameraRig))
 		{
-			Initialize(LoadedCameraRig, InCameraParameterName, BlendableParameter->ParameterType, BlendableParameter->BlendableStructType);
+			if (const UCameraObjectInterfaceBlendableParameter* BlendableParameter = LoadedCameraRig->Interface.FindBlendableParameterByName(InCameraParameterName))
+			{
+				Initialize(LoadedCameraRig, InCameraParameterName, BlendableParameter->ParameterType, BlendableParameter->BlendableStructType);
+			}
+			else if (const UCameraObjectInterfaceDataParameter* DataParameter = LoadedCameraRig->Interface.FindDataParameterByName(InCameraParameterName))
+			{
+				Initialize(LoadedCameraRig, InCameraParameterName, DataParameter->DataType, DataParameter->DataContainerType, DataParameter->DataTypeObject);
+			}
+			// else, no parameter of that name found...
 		}
-		else if (const UCameraObjectInterfaceDataParameter* DataParameter = LoadedCameraRig->Interface.FindDataParameterByName(InCameraParameterName))
-		{
-			Initialize(LoadedCameraRig, InCameraParameterName, DataParameter->DataType, DataParameter->DataContainerType, DataParameter->DataTypeObject);
-		}
-		// else, no parameter of that name found...
 	}
 }
 
diff --git a/Engine/Plugins/Cameras/GameplayCameras/Source/GameplayCamerasUncookedOnly/Private/BlueprintGraph/K2Node_SingleCameraRigParameterBase.h b/Engine/Plugins/Cameras/GameplayCameras/Source/GameplayCamerasUncookedOnly/Private/BlueprintGraph/K2Node_SingleCameraRigParameterBase.h
--- a/Engine/Plugins/Cameras/GameplayCameras/Source/GameplayCamerasUncookedOnly/Private/BlueprintGraph/K2Node_SingleCameraRigParameterBase.h
+++ b/Engine/Plugins/Cameras/GameplayCameras/Source/GameplayCamerasUncookedOnly/Private/BlueprintGraph/K2Node_SingleCameraRigParameterBase.h
@@ -29,7 +29,7 @@ public:
 
 	UK2Node_SingleCameraRigParameterBase(const FObjectInitializer& ObjectInit);
 
-	void Initialize(const FAssetData& UnloadedCameraRig, const FString& InCameraParameterName);
+	void Initialize(const FAssetData& UnloadedCameraRig, const FString& InCameraParameterName, bool bIsTemplateNode);
 	void Initialize(UCameraRigAsset* InCameraRig, const FString& InCameraParameterName, ECameraVariableType InCameraVariableType, const UScriptStruct* InBlendableStructType);
 	void Initialize(UCameraRigAsset* InCameraRig, const FString& InCameraParameterName, ECameraContextDataType InCameraContextDataType, ECameraContextDataContainerType InCameraContextDataContainerType, const UObject* InCameraContextDataTypeObject);
 
