Cameras: fix missing sub-object destruction in gameplay camera component

diff --git a/Engine/Plugins/Cameras/GameplayCameras/Source/GameplayCameras/Private/GameFramework/GameplayCameraComponentBase.cpp b/Engine/Plugins/Cameras/GameplayCameras/Source/GameplayCameras/Private/GameFramework/GameplayCameraComponentBase.cpp
--- a/Engine/Plugins/Cameras/GameplayCameras/Source/GameplayCameras/Private/GameFramework/GameplayCameraComponentBase.cpp
+++ b/Engine/Plugins/Cameras/GameplayCameras/Source/GameplayCameras/Private/GameFramework/GameplayCameraComponentBase.cpp
@@ -46,9 +46,6 @@ UGameplayCameraComponentBase::UGameplayCameraComponentBase(const FObjectInitiali
 
 	PrimaryComponentTick.bCanEverTick = true;
 	PrimaryComponentTick.TickGroup = TG_PostPhysics;
-
-	OutputCameraComponent = ObjectInit.CreateDefaultSubobject<UCineCameraComponent>(this, TEXT("OutputCameraComponent"), true);
-	OutputCameraComponent->SetupAttachment(this);
 }
 
 void UGameplayCameraComponentBase::BeginDestroy()
@@ -448,6 +445,14 @@ void UGameplayCameraComponentBase::OnRegister()
 
 	Super::OnRegister();
 
+	if (OutputCameraComponent == nullptr)
+	{
+		OutputCameraComponent = NewObject<UCineCameraComponent>(this, TEXT("OutputCameraComponent"), RF_Transient | RF_TextExportTransient);
+		OutputCameraComponent->SetupAttachment(this);
+		OutputCameraComponent->CreationMethod = CreationMethod;
+		OutputCameraComponent->RegisterComponentWithWorld(GetWorld());
+	}
+
 #if WITH_EDITOR
 
 	UWorld* World = GetWorld();
@@ -492,6 +497,16 @@ void UGameplayCameraComponentBase::EndPlay(const EEndPlayReason::Type EndPlayRea
 	Super::EndPlay(EndPlayReason);
 }
 
+void UGameplayCameraComponentBase::OnComponentDestroyed(bool bDestroyingHierarchy)
+{
+	if (OutputCameraComponent)
+	{
+		OutputCameraComponent->DestroyComponent();
+	}
+
+	Super::OnComponentDestroyed(bDestroyingHierarchy);
+}
+
 void UGameplayCameraComponentBase::Activate(bool bReset)
 {
 	// When auto-activing, this method gets called during OnRegister, before we have started playing.
diff --git a/Engine/Plugins/Cameras/GameplayCameras/Source/GameplayCameras/Public/GameFramework/GameplayCameraComponentBase.h b/Engine/Plugins/Cameras/GameplayCameras/Source/GameplayCameras/Public/GameFramework/GameplayCameraComponentBase.h
--- a/Engine/Plugins/Cameras/GameplayCameras/Source/GameplayCameras/Public/GameFramework/GameplayCameraComponentBase.h
+++ b/Engine/Plugins/Cameras/GameplayCameras/Source/GameplayCameras/Public/GameFramework/GameplayCameraComponentBase.h
@@ -147,6 +147,7 @@ public:
 	virtual void OnRegister() override;
 	virtual void BeginPlay() override;
 	virtual void EndPlay(const EEndPlayReason::Type EndPlayReason) override;
+	virtual void OnComponentDestroyed(bool bDestroyingHierarchy) override;
 	virtual void Activate(bool bReset = false) override;
 	virtual void Deactivate() override;
 	virtual void TickComponent(float DeltaTime, ELevelTick TickType, FActorComponentTickFunction *ThisTickFunction) override;
