Cameras: fixup default parameters on load as a workaround for a bug in FInstancedPropertyBags

diff --git a/Engine/Plugins/Cameras/GameplayCameras/Source/GameplayCameras/Private/Build/CameraObjectInterfaceParameterBuilder.cpp b/Engine/Plugins/Cameras/GameplayCameras/Source/GameplayCameras/Private/Build/CameraObjectInterfaceParameterBuilder.cpp
--- a/Engine/Plugins/Cameras/GameplayCameras/Source/GameplayCameras/Private/Build/CameraObjectInterfaceParameterBuilder.cpp
+++ b/Engine/Plugins/Cameras/GameplayCameras/Source/GameplayCameras/Private/Build/CameraObjectInterfaceParameterBuilder.cpp
@@ -8,7 +8,6 @@
 #include "Core/CameraParameters.h"  // IWYU pragma: keep
 #include "Core/CameraVariableReferences.h"  // IWYU pragma: keep
 #include "Core/ICustomCameraNodeParameterProvider.h"
-#include "Misc/EngineVersionComparison.h"
 #include "StructUtils/PropertyBag.h"
 #include "UObject/UnrealType.h"
 
@@ -413,5 +412,64 @@ void FCameraObjectInterfaceParameterBuilder::SetDefaultParameterValue(const UCam
 	}
 }
 
+#if UE_VERSION_OLDER_THAN(5,8,0)
+
+void FCameraObjectInterfaceParameterBuilder::FixUpDefaultParameterProperties(TConstArrayView<FCameraObjectInterfaceParameterDefinition> ParameterDefinitions, FInstancedPropertyBag& InOutPropertyBag)
+{
+	const UPropertyBag* PropertyBag = InOutPropertyBag.GetPropertyBagStruct();
+	if (!PropertyBag)
+	{
+		return;
+	}
+
+	// Before UE 5.8, there was a bug with FInstancedPropertyBag not saving the PropertyFlags of its UPropertyBag
+	// struct, resulting in camera parameters losing their CPF_Interp flag. This meant that rebuilding a camera
+	// asset would create a new parameter struct that was slightly different from the loaded one (it had those
+	// flags). This made camera assets always dirty and in need of re-saving.
+	//
+	// This function aims to fix that as a band-aid: it gets called on PostLoad by camera assets and re-adds the
+	// missing flags.
+
+	bool bFixedAny = false;
+	TArray<FPropertyBagPropertyDesc> FixedPropertyDescs(PropertyBag->GetPropertyDescs());
+
+	for (const FCameraObjectInterfaceParameterDefinition& Definition : ParameterDefinitions)
+	{
+		FPropertyBagPropertyDesc* PropertyDesc = FixedPropertyDescs.FindByPredicate(
+				[&Definition](FPropertyBagPropertyDesc& Item)
+				{
+					return Definition.ParameterGuid == Item.ID;
+				});
+		if (!PropertyDesc)
+		{
+			continue;
+		}
+
+		// Fixup blendable struct properties, and any data properties that aren't a struct.
+		bool bFixPropertyDesc = false;
+		if (Definition.ParameterType == ECameraObjectInterfaceParameterType::Blendable)
+		{
+			bFixPropertyDesc = (Definition.VariableType == ECameraVariableType::BlendableStruct);
+		}
+		else if (Definition.ParameterType == ECameraObjectInterfaceParameterType::Data)
+		{
+			bFixPropertyDesc = (Definition.DataType != ECameraContextDataType::Struct);
+		}
+		if (bFixPropertyDesc)
+		{
+			bFixedAny |= ((PropertyDesc->PropertyFlags & CPF_Interp) == 0);
+			PropertyDesc->PropertyFlags |= CPF_Interp;
+		}
+	}
+
+	if (bFixedAny)
+	{
+		const UPropertyBag* FixedPropertyBag = UPropertyBag::GetOrCreateFromDescs(FixedPropertyDescs);
+		InOutPropertyBag.MigrateToNewBagStruct(FixedPropertyBag);
+	}
+}
+
+#endif
+
 }  // namespace UE::Cameras
 
diff --git a/Engine/Plugins/Cameras/GameplayCameras/Source/GameplayCameras/Private/Core/BaseCameraObject.cpp b/Engine/Plugins/Cameras/GameplayCameras/Source/GameplayCameras/Private/Core/BaseCameraObject.cpp
--- a/Engine/Plugins/Cameras/GameplayCameras/Source/GameplayCameras/Private/Core/BaseCameraObject.cpp
+++ b/Engine/Plugins/Cameras/GameplayCameras/Source/GameplayCameras/Private/Core/BaseCameraObject.cpp
@@ -2,6 +2,9 @@
 
 #include "Core/BaseCameraObject.h"
 
+#include "Build/CameraObjectInterfaceParameterBuilder.h"
+#include "Misc/EngineVersionComparison.h"
+
 #include UE_INLINE_GENERATED_CPP_BY_NAME(BaseCameraObject)
 
 void FCameraObjectAllocationInfo::Append(const FCameraObjectAllocationInfo& OtherAllocationInfo)
@@ -17,3 +20,13 @@ void FCameraObjectAllocationInfo::Append(const FCameraObjectAllocationInfo& Othe
 	ContextDataTableInfo.DataDefinitions.Append(OtherContextDataTableInfo.DataDefinitions);
 }
 
+void UBaseCameraObject::PostLoad()
+{
+	Super::PostLoad();
+
+#if UE_VERSION_OLDER_THAN(5,8,0)
+	using namespace UE::Cameras;
+	FCameraObjectInterfaceParameterBuilder::FixUpDefaultParameterProperties(ParameterDefinitions, DefaultParameters);
+#endif
+}
+
diff --git a/Engine/Plugins/Cameras/GameplayCameras/Source/GameplayCameras/Private/Core/CameraAsset.cpp b/Engine/Plugins/Cameras/GameplayCameras/Source/GameplayCameras/Private/Core/CameraAsset.cpp
--- a/Engine/Plugins/Cameras/GameplayCameras/Source/GameplayCameras/Private/Core/CameraAsset.cpp
+++ b/Engine/Plugins/Cameras/GameplayCameras/Source/GameplayCameras/Private/Core/CameraAsset.cpp
@@ -4,9 +4,10 @@
 
 #include "Build/CameraAssetBuilder.h"
 #include "Build/CameraBuildLog.h"
+#include "Build/CameraObjectInterfaceParameterBuilder.h"
 #include "Core/CameraDirector.h"
 #include "Core/CameraRigAsset.h"
-#include "Core/CameraRigProxyAsset.h"
+#include "Misc/EngineVersionComparison.h"
 #include "Misc/EngineVersionComparison.h"
 #include "UObject/AssetRegistryTagsContext.h"
 #include "UObject/ObjectRedirector.h"
@@ -94,6 +95,11 @@ void UCameraAsset::PostLoad()
 {
 	Super::PostLoad();
 
+#if UE_VERSION_OLDER_THAN(5,8,0)
+	using namespace UE::Cameras;
+	FCameraObjectInterfaceParameterBuilder::FixUpDefaultParameterProperties(ParameterDefinitions, DefaultParameters);
+#endif
+
 #if WITH_EDITOR
 	if (CameraDirector)
 	{
diff --git a/Engine/Plugins/Cameras/GameplayCameras/Source/GameplayCameras/Public/Build/CameraObjectInterfaceParameterBuilder.h b/Engine/Plugins/Cameras/GameplayCameras/Source/GameplayCameras/Public/Build/CameraObjectInterfaceParameterBuilder.h
--- a/Engine/Plugins/Cameras/GameplayCameras/Source/GameplayCameras/Public/Build/CameraObjectInterfaceParameterBuilder.h
+++ b/Engine/Plugins/Cameras/GameplayCameras/Source/GameplayCameras/Public/Build/CameraObjectInterfaceParameterBuilder.h
@@ -3,6 +3,7 @@
 #pragma once
 
 #include "Containers/Array.h"
+#include "Misc/EngineVersionComparison.h"
 
 #define UE_API GAMEPLAYCAMERAS_API
 
@@ -40,6 +41,10 @@ public:
 	static UE_API void AppendDefaultParameterProperties(TConstArrayView<FCameraObjectInterfaceParameterDefinition> ParameterDefinitions, TArray<FPropertyBagPropertyDesc>& OutProperties);
 	static UE_API void SetDefaultParameterValues(const UBaseCameraObject* CameraObject, FInstancedPropertyBag& PropertyBag);
 
+#if UE_VERSION_OLDER_THAN(5,8,0)
+	static void FixUpDefaultParameterProperties(TConstArrayView<FCameraObjectInterfaceParameterDefinition> ParameterDefinitions, FInstancedPropertyBag& InOutPropertyBag);
+#endif
+
 private:
 
 	static void SetDefaultParameterValue(const UCameraObjectInterfaceDataParameter* DataParameter, void* DestValuePtr, const void* SrcValuePtr);
diff --git a/Engine/Plugins/Cameras/GameplayCameras/Source/GameplayCameras/Public/Core/BaseCameraObject.h b/Engine/Plugins/Cameras/GameplayCameras/Source/GameplayCameras/Public/Core/BaseCameraObject.h
--- a/Engine/Plugins/Cameras/GameplayCameras/Source/GameplayCameras/Public/Core/BaseCameraObject.h
+++ b/Engine/Plugins/Cameras/GameplayCameras/Source/GameplayCameras/Public/Core/BaseCameraObject.h
@@ -106,6 +106,11 @@ public:
 	/** Get the root node of this camera object. */
 	virtual UCameraNode* GetRootNode() { return nullptr; }
 
+public:
+
+	// UObject interface.
+	virtual void PostLoad() override;
+
 private:
 
 	/** The default interface parameter values, generated during build. */
