Cameras: fix GetActiveCameraRigEvaluationInfo returning the wrong information when rigs are blending

diff --git a/Engine/Plugins/Cameras/GameplayCameras/Source/GameplayCameras/Private/Core/BlendStackCameraNode.cpp b/Engine/Plugins/Cameras/GameplayCameras/Source/GameplayCameras/Private/Core/BlendStackCameraNode.cpp
--- a/Engine/Plugins/Cameras/GameplayCameras/Source/GameplayCameras/Private/Core/BlendStackCameraNode.cpp
+++ b/Engine/Plugins/Cameras/GameplayCameras/Source/GameplayCameras/Private/Core/BlendStackCameraNode.cpp
@@ -132,7 +132,10 @@ void FBlendStackCameraNodeEvaluator::InitializeEntry(
 		FCameraNodeEvaluatorInitializeParams InitParams(&NewEntry.EvaluatorHierarchy);
 		InitParams.Evaluator = OwningEvaluator;
 		InitParams.EvaluationContext = EvaluationContext;
-		InitParams.LastActiveCameraRigInfo = GetActiveCameraRigEvaluationInfo();
+		if (!bSetActiveResult)
+		{
+			InitParams.LastActiveCameraRigInfo = GetActiveCameraRigEvaluationInfo();
+		}
 		InitParams.Layer = Layer;
 		RootEvaluator->Initialize(InitParams, NewEntry.ContextResult);  // Initializing with the context result here.
 	}
@@ -205,7 +208,7 @@ FCameraRigEvaluationInfo FBlendStackCameraNodeEvaluator::GetActiveCameraRigEvalu
 {
 	if (Entries.Num() > 0)
 	{
-		const FCameraRigEntry& ActiveEntry = Entries[0];
+		const FCameraRigEntry& ActiveEntry = Entries.Last();
 		FCameraRigEvaluationInfo Info(
 				FCameraRigInstanceID::FromBlendStackEntryID(ActiveEntry.EntryID, Layer),
 				ActiveEntry.EvaluationContext.Pin(),
