Cameras: get rid of the debug info if we don't need it anymore

diff --git a/Engine/Plugins/Cameras/GameplayCameras/Source/GameplayCameras/Private/Core/CameraSystemEvaluator.cpp b/Engine/Plugins/Cameras/GameplayCameras/Source/GameplayCameras/Private/Core/CameraSystemEvaluator.cpp
--- a/Engine/Plugins/Cameras/GameplayCameras/Source/GameplayCameras/Private/Core/CameraSystemEvaluator.cpp
+++ b/Engine/Plugins/Cameras/GameplayCameras/Source/GameplayCameras/Private/Core/CameraSystemEvaluator.cpp
@@ -570,6 +570,7 @@ void FCameraSystemEvaluator::BuildDebugBlocksIfNeeded()
 {
 	if (!ShouldBuildOrDrawDebugBlocks())
 	{
+		RootDebugBlock = nullptr;
 		return;
 	}
 
