Cameras: use the correct values when a camera rig parameter override is itself driven by a higher-level camera rig parameter override

diff --git a/Engine/Plugins/Cameras/GameplayCameras/Source/GameplayCameras/Private/Helpers/CameraObjectInterfaceParameterOverrideHelper.cpp b/Engine/Plugins/Cameras/GameplayCameras/Source/GameplayCameras/Private/Helpers/CameraObjectInterfaceParameterOverrideHelper.cpp
--- a/Engine/Plugins/Cameras/GameplayCameras/Source/GameplayCameras/Private/Helpers/CameraObjectInterfaceParameterOverrideHelper.cpp
+++ b/Engine/Plugins/Cameras/GameplayCameras/Source/GameplayCameras/Private/Helpers/CameraObjectInterfaceParameterOverrideHelper.cpp
@@ -26,15 +26,23 @@ void ApplyBlendableParameterOverride(
 	using ValueType = typename ParameterType::ValueType;
 
 	const FCameraVariableID ParameterVariableID(ParameterDefinition.VariableID);
-	if (ParameterValue.Variable != nullptr)
+	if (ParameterValue.VariableID.IsValid())
 	{
 		// The override is driven by a variable... read its value and set it as the value for the
 		// prefab's variable. Basically, we forward the value from one variable to the next.
-		FCameraVariableDefinition OverrideDefinition(ParameterValue.Variable->GetVariableDefinition());
-
-		const ValueType OverrideValue = VariableTable.GetValue<ValueType>(
-				OverrideDefinition.VariableID, ParameterValue.Variable->GetDefaultValue());
-		VariableTable.SetValue<ValueType>(ParameterVariableID, OverrideValue);
+		ValueType OverrideValue;
+		if (VariableTable.TryGetValue<ValueType>(ParameterValue.VariableID, OverrideValue))
+		{
+			VariableTable.SetValue<ValueType>(ParameterVariableID, OverrideValue);
+		}
+		else if (ParameterValue.Variable)
+		{
+			VariableTable.SetValue<ValueType>(ParameterVariableID, ParameterValue.Variable->GetDefaultValue());
+		}
+		else
+		{
+			VariableTable.SetValue<ValueType>(ParameterVariableID, ParameterValue.Value);
+		}
 	}
 	else if (!bDrivenOnly)
 	{
