Expose some functions in GameplayCameraActor and GameplayCameraActorBase - Needed these for editor previews in the scene graph.

diff --git a/Engine/Plugins/Cameras/GameplayCameras/Source/GameplayCameras/Public/GameFramework/GameplayCameraActor.h b/Engine/Plugins/Cameras/GameplayCameras/Source/GameplayCameras/Public/GameFramework/GameplayCameraActor.h
--- a/Engine/Plugins/Cameras/GameplayCameras/Source/GameplayCameras/Public/GameFramework/GameplayCameraActor.h
+++ b/Engine/Plugins/Cameras/GameplayCameras/Source/GameplayCameras/Public/GameFramework/GameplayCameraActor.h
@@ -18,7 +18,7 @@ class AGameplayCameraActor : public AGameplayCameraActorBase
 
 public:
 
-	AGameplayCameraActor(const FObjectInitializer& ObjectInit);
+	GAMEPLAYCAMERAS_API AGameplayCameraActor(const FObjectInitializer& ObjectInit);
 
 public:
 
@@ -29,12 +29,12 @@ public:
 public:
 
 	// AActor interface.
-	virtual USceneComponent* GetDefaultAttachComponent() const override;
+	GAMEPLAYCAMERAS_API virtual USceneComponent* GetDefaultAttachComponent() const override;
 
 protected:
 
 	// AGameplayCameraActorBase interface.
-	virtual UGameplayCameraComponentBase* GetCameraComponentBase() const override;
+	GAMEPLAYCAMERAS_API virtual UGameplayCameraComponentBase* GetCameraComponentBase() const override;
 
 private:
 
diff --git a/Engine/Plugins/Cameras/GameplayCameras/Source/GameplayCameras/Public/GameFramework/GameplayCameraActorBase.h b/Engine/Plugins/Cameras/GameplayCameras/Source/GameplayCameras/Public/GameFramework/GameplayCameraActorBase.h
--- a/Engine/Plugins/Cameras/GameplayCameras/Source/GameplayCameras/Public/GameFramework/GameplayCameraActorBase.h
+++ b/Engine/Plugins/Cameras/GameplayCameras/Source/GameplayCameras/Public/GameFramework/GameplayCameraActorBase.h
@@ -29,7 +29,7 @@ public:
 public:
 
 	// AActor interface.
-	virtual void CalcCamera(float DeltaTime, FMinimalViewInfo& OutResult) override;
+	GAMEPLAYCAMERAS_API virtual void CalcCamera(float DeltaTime, FMinimalViewInfo& OutResult) override;
 
 protected:
 
