Cameras: add BP node function to get the owning rig

diff --git a/Engine/Plugins/Cameras/GameplayCameras/Source/GameplayCameras/Private/Nodes/Utility/BlueprintCameraNode.cpp b/Engine/Plugins/Cameras/GameplayCameras/Source/GameplayCameras/Private/Nodes/Utility/BlueprintCameraNode.cpp
--- a/Engine/Plugins/Cameras/GameplayCameras/Source/GameplayCameras/Private/Nodes/Utility/BlueprintCameraNode.cpp
+++ b/Engine/Plugins/Cameras/GameplayCameras/Source/GameplayCameras/Private/Nodes/Utility/BlueprintCameraNode.cpp
@@ -275,6 +275,11 @@ void UBlueprintCameraNodeEvaluator::SetDefaultOwningCameraRigParameters(FBluepri
 	}
 }
 
+UCameraRigAsset* UBlueprintCameraNodeEvaluator::GetCameraRig() const
+{
+	return BlueprintNode->GetTypedOuter<UCameraRigAsset>();
+}
+
 APlayerController* UBlueprintCameraNodeEvaluator::GetPlayerController() const
 {
 	if (CurrentContext)
diff --git a/Engine/Plugins/Cameras/GameplayCameras/Source/GameplayCameras/Public/Nodes/Utility/BlueprintCameraNode.h b/Engine/Plugins/Cameras/GameplayCameras/Source/GameplayCameras/Public/Nodes/Utility/BlueprintCameraNode.h
--- a/Engine/Plugins/Cameras/GameplayCameras/Source/GameplayCameras/Public/Nodes/Utility/BlueprintCameraNode.h
+++ b/Engine/Plugins/Cameras/GameplayCameras/Source/GameplayCameras/Public/Nodes/Utility/BlueprintCameraNode.h
@@ -60,6 +60,10 @@ public:
 	UFUNCTION(BlueprintCallable, Category="Evaluation")
 	void SetDefaultOwningCameraRigParameters(FBlueprintCameraEvaluationDataRef TargetCameraData) const;
 
+	/* Gets the owning camera rig for this node. */
+	UFUNCTION(BlueprintPure, Category="Evaluation")
+	UCameraRigAsset* GetCameraRig() const;
+
 	/** Gets the player controller that the node is running for, if any. */
 	UFUNCTION(BlueprintPure, Category="Evaluation")
 	APlayerController* GetPlayerController() const;
