Cameras: add output return value pin to BP parameter setter node

diff --git a/Engine/Plugins/Cameras/GameplayCameras/Source/GameplayCamerasUncookedOnly/Private/BlueprintGraph/K2Node_SetCameraRigParameter.cpp b/Engine/Plugins/Cameras/GameplayCameras/Source/GameplayCamerasUncookedOnly/Private/BlueprintGraph/K2Node_SetCameraRigParameter.cpp
--- a/Engine/Plugins/Cameras/GameplayCameras/Source/GameplayCamerasUncookedOnly/Private/BlueprintGraph/K2Node_SetCameraRigParameter.cpp
+++ b/Engine/Plugins/Cameras/GameplayCameras/Source/GameplayCamerasUncookedOnly/Private/BlueprintGraph/K2Node_SetCameraRigParameter.cpp
@@ -19,9 +19,11 @@ void UK2Node_SetCameraRigParameter::AllocateDefaultPins()
 {
 	Super::AllocateDefaultPins();
 
-	// Add parameter value pin.
+	// Add parameter value pin and return pin.
 	FEdGraphPinType PinType = GetParameterPinType();
 	CreatePin(EGPD_Input, PinType, FName(CameraParameterName));
+	UEdGraphPin* ReturnPin = CreatePin(EGPD_Output, PinType, UEdGraphSchema_K2::PN_ReturnValue);
+	ReturnPin->PinFriendlyName = FText::GetEmpty();
 }
 
 FText UK2Node_SetCameraRigParameter::GetNodeTitle(ENodeTitleType::Type TitleType) const
@@ -81,6 +83,7 @@ void UK2Node_SetCameraRigParameter::ExpandNode(FKismetCompilerContext& CompilerC
 
 	UEdGraphPin* const CameraNodeEvaluationResultPin = GetCameraNodeEvaluationResultPin();
 	UEdGraphPin* const CameraParameterValuePin = FindPinChecked(CameraParameterName);
+	UEdGraphPin* const ReturnValuePin = FindPinChecked(UEdGraphSchema_K2::PN_ReturnValue);
 
 	// Make the SetXxxParameter function call node.
 	UK2Node_CallFunction* CallSetParameter = CompilerContext.SpawnIntermediateNode<UK2Node_CallFunction>(this, SourceGraph);
@@ -111,8 +114,8 @@ void UK2Node_SetCameraRigParameter::ExpandNode(FKismetCompilerContext& CompilerC
 	}
 	else if (UK2Node* MakeLiteral = MakeLiteralValueForPin(CompilerContext, SourceGraph, this, CameraParameterValuePin))
 	{
-		UEdGraphPin* ReturnValuePin = MakeLiteral->FindPinChecked(UEdGraphSchema_K2::PN_ReturnValue);
-		ReturnValuePin->MakeLinkTo(CallSetParameterValuePin);
+		UEdGraphPin* LiteralValuePin = MakeLiteral->FindPinChecked(UEdGraphSchema_K2::PN_ReturnValue);
+		LiteralValuePin->MakeLinkTo(CallSetParameterValuePin);
 	}
 	else
 	{
@@ -125,6 +128,15 @@ void UK2Node_SetCameraRigParameter::ExpandNode(FKismetCompilerContext& CompilerC
 		CompilerContext.MessageLog.Error(*MissingParameterConnectionMsg.ToString(), this);
 	}
 
+	// Connect the pin that provides the value to the return pin of our source node.
+	for (UEdGraphPin* InValuePin : CallSetParameterValuePin->LinkedTo)
+	{
+		for (UEdGraphPin* OutValuePin : ReturnValuePin->LinkedTo)
+		{
+			InValuePin->MakeLinkTo(OutValuePin);
+		}
+	}
+	
 	// Setup the execution flow.
 	UEdGraphPin* ThisExecPin = GetExecPin();
 	CompilerContext.MovePinLinksToIntermediate(*ThisExecPin, *FirstExecPin);
