Cameras: add more Blueprint methods to the components for managing persistent rigs and shakes

diff --git a/Engine/Plugins/Cameras/GameplayCameras/Source/GameplayCameras/Private/GameFramework/GameplayCameraComponentBase.cpp b/Engine/Plugins/Cameras/GameplayCameras/Source/GameplayCameras/Private/GameFramework/GameplayCameraComponentBase.cpp
--- a/Engine/Plugins/Cameras/GameplayCameras/Source/GameplayCameras/Private/GameFramework/GameplayCameraComponentBase.cpp
+++ b/Engine/Plugins/Cameras/GameplayCameras/Source/GameplayCameras/Private/GameFramework/GameplayCameraComponentBase.cpp
@@ -20,6 +20,8 @@
 #include "Misc/AssertionMacros.h"
 #include "Misc/EngineVersionComparison.h"
 #include "SceneView.h"
+#include "Services/CameraModifierService.h"
+#include "Services/CameraShakeService.h"
 #include "ShowFlags.h"
 
 #if UE_VERSION_NEWER_THAN_OR_EQUAL(5,6,0)
@@ -430,6 +432,86 @@ void UGameplayCameraComponentBase::ActivatePersistentVisualCameraRig(UCameraRigA
 
 #undef UE_PRIVATE_GAMEPLAY_CAMERA_COMPONENT_ACTIVATE_NON_MAIN_RIG
 
+FCameraRigInstanceID UGameplayCameraComponentBase::StartGlobalCameraModifierRig(const UCameraRigAsset* CameraRig, int32 OrderKey)
+{
+	using namespace UE::Cameras;
+
+	if (CameraSystemEvaluator)
+	{
+		TSharedPtr<FCameraModifierService> CameraModifierService = CameraSystemEvaluator->FindEvaluationService<FCameraModifierService>();
+		return CameraModifierService->StartCameraModifierRig(CameraRig, EvaluationContext.ToSharedRef(), ECameraRigLayer::Global, OrderKey);
+	}
+
+	return FCameraRigInstanceID();
+}
+
+FCameraRigInstanceID UGameplayCameraComponentBase::StartVisualCameraModifierRig(const UCameraRigAsset* CameraRig, int32 OrderKey)
+{
+	using namespace UE::Cameras;
+
+	if (CameraSystemEvaluator)
+	{
+		TSharedPtr<FCameraModifierService> CameraModifierService = CameraSystemEvaluator->FindEvaluationService<FCameraModifierService>();
+		return CameraModifierService->StartCameraModifierRig(CameraRig, EvaluationContext.ToSharedRef(), ECameraRigLayer::Visual, OrderKey);
+	}
+
+	return FCameraRigInstanceID();
+}
+
+void UGameplayCameraComponentBase::StopCameraModifierRig(FCameraRigInstanceID InstanceID, bool bImmediately)
+{
+	using namespace UE::Cameras;
+
+	if (CameraSystemEvaluator)
+	{
+		TSharedPtr<FCameraModifierService> CameraModifierService = CameraSystemEvaluator->FindEvaluationService<FCameraModifierService>();
+		return CameraModifierService->StopCameraModifierRig(InstanceID, bImmediately);
+	}
+}
+
+FCameraShakeInstanceID UGameplayCameraComponentBase::StartCameraShakeAsset(const UCameraShakeAsset* CameraShake, float ShakeScale, ECameraShakePlaySpace PlaySpace, FRotator UserPlaySpaceRotation)
+{
+	using namespace UE::Cameras;
+
+	if (CameraSystemEvaluator)
+	{
+		FStartCameraShakeParams Params;
+		Params.CameraShake = CameraShake;
+		Params.ShakeScale = ShakeScale;
+		Params.PlaySpace = PlaySpace;
+		Params.UserPlaySpaceRotation = UserPlaySpaceRotation;
+
+		TSharedPtr<FCameraShakeService> CameraShakeService = CameraSystemEvaluator->FindEvaluationService<FCameraShakeService>();
+		return CameraShakeService->StartCameraShake(Params);
+	}
+
+	return FCameraShakeInstanceID();
+}
+
+bool UGameplayCameraComponentBase::IsCameraShakeAssetPlaying(FCameraShakeInstanceID InInstanceID) const
+{
+	using namespace UE::Cameras;
+
+	if (CameraSystemEvaluator)
+	{
+		TSharedPtr<FCameraShakeService> CameraShakeService = CameraSystemEvaluator->FindEvaluationService<FCameraShakeService>();
+		return CameraShakeService->IsCameraShakePlaying(InInstanceID);
+	}
+	return false;
+}
+
+bool UGameplayCameraComponentBase::StopCameraShakeAsset(FCameraShakeInstanceID InInstanceID, bool bImmediately)
+{
+	using namespace UE::Cameras;
+
+	if (CameraSystemEvaluator)
+	{
+		TSharedPtr<FCameraShakeService> CameraShakeService = CameraSystemEvaluator->FindEvaluationService<FCameraShakeService>();
+		return CameraShakeService->StopCameraShake(InInstanceID, bImmediately);
+	}
+	return false;
+}
+
 FRotator UGameplayCameraComponentBase::GetEvaluatedCameraRotation() const
 {
 	using namespace UE::Cameras;
diff --git a/Engine/Plugins/Cameras/GameplayCameras/Source/GameplayCameras/Public/GameFramework/GameplayCameraComponentBase.h b/Engine/Plugins/Cameras/GameplayCameras/Source/GameplayCameras/Public/GameFramework/GameplayCameraComponentBase.h
--- a/Engine/Plugins/Cameras/GameplayCameras/Source/GameplayCameras/Public/GameFramework/GameplayCameraComponentBase.h
+++ b/Engine/Plugins/Cameras/GameplayCameras/Source/GameplayCameras/Public/GameFramework/GameplayCameraComponentBase.h
@@ -5,6 +5,7 @@
 #include "Components/SceneComponent.h"
 #include "Core/CameraAssetReference.h"
 #include "Core/CameraEvaluationContext.h"
+#include "Core/CameraShakeInstanceID.h"
 #include "GameFramework/BlueprintCameraEvaluationDataRef.h"
 #include "GameFramework/IGameplayCameraSystemHost.h"
 #include "UObject/ObjectMacros.h"
@@ -16,6 +17,7 @@ class FPrimitiveDrawInterface;
 class FSceneView;
 class FViewport;
 class UCameraAsset;
+class UCameraShakeAsset;
 class UCanvas;
 class UCineCameraComponent;
 struct FCameraContextDataTableAllocationInfo;
@@ -141,6 +143,38 @@ public:
 	UFUNCTION(BlueprintCallable, Category="Camera")
 	void ActivatePersistentVisualCameraRig(UCameraRigAsset* CameraRig);
 
+public:
+
+	/** Starts a camera modifier rig on the global layer. */
+	UFUNCTION(BlueprintCallable, Category="Camera")
+	FCameraRigInstanceID StartGlobalCameraModifierRig(const UCameraRigAsset* CameraRig, int32 OrderKey = 0);
+
+	/** Starts a camera modifier rig on the visual layer. */
+	UFUNCTION(BlueprintCallable, Category="Camera")
+	FCameraRigInstanceID StartVisualCameraModifierRig(const UCameraRigAsset* CameraRig, int32 OrderKey = 0);
+
+	/** Stops a camera modifier rig on previously started on the global or visual layer. */
+	UFUNCTION(BlueprintCallable, Category="Camera")
+	void StopCameraModifierRig(FCameraRigInstanceID InstanceID, bool bImmediately = false);
+
+public:
+
+	/** Starts a new camera shake. */
+	UFUNCTION(BlueprintCallable, Category="Camera")
+	FCameraShakeInstanceID StartCameraShakeAsset(
+			const UCameraShakeAsset* CameraShake,
+			float ShakeScale = 1.f,
+			ECameraShakePlaySpace PlaySpace = ECameraShakePlaySpace::CameraLocal,
+			FRotator UserPlaySpaceRotation = FRotator::ZeroRotator);
+
+	/** Checks if a camera shake is running. */
+	UFUNCTION(BlueprintCallable, Category="Camera")
+	bool IsCameraShakeAssetPlaying(FCameraShakeInstanceID InInstanceID) const;
+
+	/** Stops a running camera shake. */
+	UFUNCTION(BlueprintCallable, Category="Camera")
+	bool StopCameraShakeAsset(FCameraShakeInstanceID InInstanceID, bool bImmediately = false);
+
 public:
 
 	// UActorComponent interface
