Cameras: make perlin noise use cubic splines by default

This fixes undesirable "mini-pauses" in the middle of some curve segments when using a standard smooth-step interpolation.

diff --git a/Engine/Plugins/Cameras/GameplayCameras/Source/GameplayCameras/Private/Math/PerlinNoise.cpp b/Engine/Plugins/Cameras/GameplayCameras/Source/GameplayCameras/Private/Math/PerlinNoise.cpp
--- a/Engine/Plugins/Cameras/GameplayCameras/Source/GameplayCameras/Private/Math/PerlinNoise.cpp
+++ b/Engine/Plugins/Cameras/GameplayCameras/Source/GameplayCameras/Private/Math/PerlinNoise.cpp
@@ -2,6 +2,7 @@
 
 #include "Math/PerlinNoise.h"
 
+#include "HAL/IConsoleManager.h"
 #include "Math/Interpolation.h"
 #include "Math/UnrealMathUtility.h"
 #include "Serialization/Archive.h"
@@ -11,6 +12,18 @@
 namespace UE::Cameras
 {
 
+bool GCubicPerlinNoise = true;
+static FAutoConsoleVariableRef CVarCubicPerlinNoise(
+	TEXT("GameplayCameras.CubicPerlinNoise"),
+	GCubicPerlinNoise,
+	TEXT("(Default: true) Enables cubic splines for perlin noise generator."));
+
+float GCubicPerlinNoiseTension = 0.1f;
+static FAutoConsoleVariableRef CVarCubicPerlinNoiseTension(
+	TEXT("GameplayCameras.CubicPerlinNoiseTension"),
+	GCubicPerlinNoiseTension,
+	TEXT("(Default: 0.5) Sets the tension for the cubic splines of the perlin noise generator."));
+
 FPerlinNoise::FPerlinNoise()
 {
 	Initialize(1.f);
@@ -38,13 +51,20 @@ void FPerlinNoise::Initialize(float InFrequency)
 	{
 		FSinglePerlinNoise& Octave(Octaves[Index]);
 		Octave.Frequency = OctaveFrequency;
+		Octave.SecondPrev = FMath::FRandRange(-1.f, 1.f);
 		Octave.Prev = FMath::FRandRange(-1.f, 1.f);
 		Octave.Next = FMath::FRandRange(-1.f, 1.f);
+		Octave.SecondNext = FMath::FRandRange(-1.f, 1.f);
 
 		OctaveFrequency *= Lacunarity;
 	}
 }
 
+FVector2f FPerlinNoise::ComputeTangent(float InPrev, float InNext, float Interval, float Tension)
+{
+	return (1.f - Tension) * (FVector2f(Interval, InNext) - FVector2f(-Interval, InPrev)) / (2.f * Interval);
+}
+
 void FPerlinNoise::SetFrequency(float InFrequency)
 {
 	if (InFrequency == Octaves[0].Frequency)
@@ -92,8 +112,10 @@ void FPerlinNoise::SetNumOctaves(uint8 InNumOctaves)
 			FSinglePerlinNoise& Octave(Octaves[Index]);
 			Octave.Frequency = OctaveFrequency;
 			Octave.CurTime = 0.f;
+			Octave.SecondPrev = FMath::FRandRange(-1.f, 1.f);
 			Octave.Prev = FMath::FRandRange(-1.f, 1.f);
 			Octave.Next = FMath::FRandRange(-1.f, 1.f);
+			Octave.SecondNext = FMath::FRandRange(-1.f, 1.f);
 
 			OctaveFrequency *= Lacunarity;
 		}
@@ -107,6 +129,8 @@ float FPerlinNoise::GenerateValue(float DeltaTime)
 	float Value = 0.f;
 	float OctaveAmplitude = Amplitude;
 
+	const float Tension = FMath::Clamp(GCubicPerlinNoiseTension, 0.f, 1.f);
+
 	for (int32 Index = 0; Index < NumOctaves; ++Index)
 	{
 		FSinglePerlinNoise& Octave(Octaves[Index]);
@@ -118,18 +142,37 @@ float FPerlinNoise::GenerateValue(float DeltaTime)
 
 		// If we are going over the end of the current interval, generate
 		// some new value for the next interval.
+		// NOTE: for large delta-times we might skip over intervals, but here we behave as if
+		//       only going to the next interval.
 		if ((int32)NextNumIntervals > (int32)PrevNumIntervals)
 		{
+			Octave.SecondPrev = Octave.Prev;
 			Octave.Prev = Octave.Next;
-			Octave.Next = FMath::FRandRange(-1.f, 1.f);
+			Octave.Next = Octave.SecondNext;
+			Octave.SecondNext = FMath::FRandRange(-1.f, 1.f);
 		}
 
 		Octave.CurTime += DeltaTime;
 
 		const float IntervalFactor = (NextNumIntervals - FMath::TruncToFloat(NextNumIntervals));
-		const float InterpFactor = SmootherStep(IntervalFactor);
-
-		Value += OctaveAmplitude * FMath::Lerp(Octave.Prev, Octave.Next, InterpFactor);
+		if (GCubicPerlinNoise)
+		{
+			// ComputeTangent returns the right-hand tangent (i.e. the tangent oriented towards the "future"),
+			// so we need to turn it around for the next point -- both tangents need to be turned "inwards" here.
+			const FVector2f PrevTangent = ComputeTangent(Octave.SecondPrev, Octave.Next, Interval, Tension);
+			const FVector2f NextTangent = -ComputeTangent(Octave.Prev, Octave.SecondNext, Interval, Tension);
+
+			const FVector2f InterpPoint = FMath::CubicInterp(
+					FVector2f(0.f, Octave.Prev), PrevTangent,
+					FVector2f(Interval, Octave.Next), -NextTangent,
+					IntervalFactor);
+			Value += OctaveAmplitude * InterpPoint.Y;
+		}
+		else
+		{
+			const float InterpFactor = SmoothStep(IntervalFactor);
+			Value += OctaveAmplitude * FMath::Lerp(Octave.Prev, Octave.Next, InterpFactor);
+		}
 
 		OctaveAmplitude *= OctaveGain;
 	}
@@ -145,13 +188,10 @@ void FPerlinNoise::Serialize(FArchive& Ar)
 
 	Ar << NumOctaves;
 
-	for (int32 Index = 0; Index < MAX_OCTAVES; ++Index)
+	for (int32 Index = 0; Index < NumOctaves; ++Index)
 	{
 		FSinglePerlinNoise& Octave(Octaves[Index]);
-		Ar << Octave.Frequency;
-		Ar << Octave.CurTime;
-		Ar << Octave.Prev;
-		Ar << Octave.Next;
+		Ar.Serialize(static_cast<void*>(&Octave), sizeof(FSinglePerlinNoise));
 	}
 }
 
diff --git a/Engine/Plugins/Cameras/GameplayCameras/Source/GameplayCameras/Public/Math/PerlinNoise.h b/Engine/Plugins/Cameras/GameplayCameras/Source/GameplayCameras/Public/Math/PerlinNoise.h
--- a/Engine/Plugins/Cameras/GameplayCameras/Source/GameplayCameras/Public/Math/PerlinNoise.h
+++ b/Engine/Plugins/Cameras/GameplayCameras/Source/GameplayCameras/Public/Math/PerlinNoise.h
@@ -75,6 +75,8 @@ public:
 private:
 
 	UE_API void Initialize(float InFrequency);
+	
+	static FVector2f ComputeTangent(float InPrev, float InNext, float Interval, float Tension);
 
 private:
 
@@ -84,8 +86,10 @@ private:
 	{
 		float Frequency = 1.f;
 		float CurTime = 0.f;
+		float SecondPrev = 0.f;
 		float Prev = 0.f;
 		float Next = 0.f;
+		float SecondNext = 0.f;
 	};
 
 	float Amplitude = 1.f;
