Cameras: fix GetActiveCameraRigEvaluationInfo returning the wrong information when rigs are blending

#FailedValidations Preflight

diff --git a/Engine/Plugins/Cameras/GameplayCameras/Source/GameplayCameras/Private/Core/BlendStackCameraNode.cpp b/Engine/Plugins/Cameras/GameplayCameras/Source/GameplayCameras/Private/Core/BlendStackCameraNode.cpp
--- a/Engine/Plugins/Cameras/GameplayCameras/Source/GameplayCameras/Private/Core/BlendStackCameraNode.cpp
+++ b/Engine/Plugins/Cameras/GameplayCameras/Source/GameplayCameras/Private/Core/BlendStackCameraNode.cpp
@@ -93,7 +93,7 @@ void FBlendStackCameraNodeEvaluator::InitializeEntry(
 		const UCameraRigAsset* CameraRig,
 		TSharedPtr<const FCameraEvaluationContext> EvaluationContext,
 		UBlendStackRootCameraNode* EntryRootNode,
-		bool bSetActiveResult)
+		bool bIsActiveEntry)
 {
 	// Clear the evaluator hierarchy in case we are hot-reloading an entry.
 	NewEntry.EvaluatorHierarchy.Reset();
@@ -116,7 +116,7 @@ void FBlendStackCameraNodeEvaluator::InitializeEntry(
 	NewEntry.ContextResult.ContextDataTable.OverrideAll(ContextResult.ContextDataTable);
 
 	// Add some conditional result if necessary.
-	if (bSetActiveResult && EvaluationContext)
+	if (bIsActiveEntry && EvaluationContext)
 	{
 		const FCameraNodeEvaluationResult* ActiveOnlyResult = EvaluationContext->GetConditionalResult(ECameraEvaluationDataCondition::ActiveCameraRig);
 		if (ActiveOnlyResult)
@@ -132,7 +132,10 @@ void FBlendStackCameraNodeEvaluator::InitializeEntry(
 		FCameraNodeEvaluatorInitializeParams InitParams(&NewEntry.EvaluatorHierarchy);
 		InitParams.Evaluator = OwningEvaluator;
 		InitParams.EvaluationContext = EvaluationContext;
-		InitParams.LastActiveCameraRigInfo = GetActiveCameraRigEvaluationInfo();
+		if (!bIsActiveEntry)
+		{
+			InitParams.LastActiveCameraRigInfo = GetActiveCameraRigEvaluationInfo();
+		}
 		InitParams.Layer = Layer;
 		RootEvaluator->Initialize(InitParams, NewEntry.ContextResult);  // Initializing with the context result here.
 	}
@@ -205,7 +208,7 @@ FCameraRigEvaluationInfo FBlendStackCameraNodeEvaluator::GetActiveCameraRigEvalu
 {
 	if (Entries.Num() > 0)
 	{
-		const FCameraRigEntry& ActiveEntry = Entries[0];
+		const FCameraRigEntry& ActiveEntry = Entries.Last();
 		FCameraRigEvaluationInfo Info(
 				FCameraRigInstanceID::FromBlendStackEntryID(ActiveEntry.EntryID, Layer),
 				ActiveEntry.EvaluationContext.Pin(),
diff --git a/Engine/Plugins/Cameras/GameplayCameras/Source/GameplayCameras/Public/Core/BlendStackCameraNode.h b/Engine/Plugins/Cameras/GameplayCameras/Source/GameplayCameras/Public/Core/BlendStackCameraNode.h
--- a/Engine/Plugins/Cameras/GameplayCameras/Source/GameplayCameras/Public/Core/BlendStackCameraNode.h
+++ b/Engine/Plugins/Cameras/GameplayCameras/Source/GameplayCameras/Public/Core/BlendStackCameraNode.h
@@ -152,7 +152,7 @@ protected:
 		const UCameraRigAsset* CameraRig,
 		TSharedPtr<const FCameraEvaluationContext> EvaluationContext,
 		UBlendStackRootCameraNode* EntryRootNode,
-		bool bSetActiveResult);
+		bool bIsActiveEntry);
 
 	int32 IndexOfEntry(const FBlendStackEntryID EntryID) const;
 
