Cameras: fix double-negation and incorrect comment

#FailedValidations Preflight

diff --git a/Engine/Plugins/Cameras/GameplayCameras/Source/GameplayCameras/Private/Math/PerlinNoise.cpp b/Engine/Plugins/Cameras/GameplayCameras/Source/GameplayCameras/Private/Math/PerlinNoise.cpp
--- a/Engine/Plugins/Cameras/GameplayCameras/Source/GameplayCameras/Private/Math/PerlinNoise.cpp
+++ b/Engine/Plugins/Cameras/GameplayCameras/Source/GameplayCameras/Private/Math/PerlinNoise.cpp
@@ -22,7 +22,7 @@ float GCubicPerlinNoiseTension = 0.1f;
 static FAutoConsoleVariableRef CVarCubicPerlinNoiseTension(
 	TEXT("GameplayCameras.CubicPerlinNoiseTension"),
 	GCubicPerlinNoiseTension,
-	TEXT("(Default: 0.5) Sets the tension for the cubic splines of the perlin noise generator."));
+	TEXT("(Default: 0.1) Sets the tension for the cubic splines of the perlin noise generator."));
 
 FPerlinNoise::FPerlinNoise()
 {
@@ -157,14 +157,12 @@ float FPerlinNoise::GenerateValue(float DeltaTime)
 		const float IntervalFactor = (NextNumIntervals - FMath::TruncToFloat(NextNumIntervals));
 		if (GCubicPerlinNoise)
 		{
-			// ComputeTangent returns the right-hand tangent (i.e. the tangent oriented towards the "future"),
-			// so we need to turn it around for the next point -- both tangents need to be turned "inwards" here.
 			const FVector2f PrevTangent = ComputeTangent(Octave.SecondPrev, Octave.Next, Interval, Tension);
-			const FVector2f NextTangent = -ComputeTangent(Octave.Prev, Octave.SecondNext, Interval, Tension);
+			const FVector2f NextTangent = ComputeTangent(Octave.Prev, Octave.SecondNext, Interval, Tension);
 
 			const FVector2f InterpPoint = FMath::CubicInterp(
 					FVector2f(0.f, Octave.Prev), PrevTangent,
-					FVector2f(Interval, Octave.Next), -NextTangent,
+					FVector2f(Interval, Octave.Next), NextTangent,
 					IntervalFactor);
 			Value += OctaveAmplitude * InterpPoint.Y;
 		}
