Cameras: cleanup live-edit node listeners after a GC pass

#FailedValidations Preflight

diff --git a/Engine/Plugins/Cameras/GameplayCameras/Source/GameplayCamerasEditor/Private/GameplayCamerasLiveEditManager.cpp b/Engine/Plugins/Cameras/GameplayCameras/Source/GameplayCamerasEditor/Private/GameplayCamerasLiveEditManager.cpp
--- a/Engine/Plugins/Cameras/GameplayCameras/Source/GameplayCamerasEditor/Private/GameplayCamerasLiveEditManager.cpp
+++ b/Engine/Plugins/Cameras/GameplayCameras/Source/GameplayCamerasEditor/Private/GameplayCamerasLiveEditManager.cpp
@@ -163,6 +163,13 @@ void FGameplayCamerasLiveEditManager::RemoveGarbage()
 			It.RemoveCurrent();
 		}
 	}
+	for (auto It = NodeListenerMap.CreateIterator(); It; ++It)
+	{
+		if (!It.Key().IsValid())
+		{
+			It.RemoveCurrent();
+		}
+	}
 }
 
 namespace Internal
