Expose constructors on FCameraRigAssetReference

diff --git a/Engine/Plugins/Cameras/GameplayCameras/Source/GameplayCameras/Public/Core/CameraRigAssetReference.h b/Engine/Plugins/Cameras/GameplayCameras/Source/GameplayCameras/Public/Core/CameraRigAssetReference.h
--- a/Engine/Plugins/Cameras/GameplayCameras/Source/GameplayCameras/Public/Core/CameraRigAssetReference.h
+++ b/Engine/Plugins/Cameras/GameplayCameras/Source/GameplayCameras/Public/Core/CameraRigAssetReference.h
@@ -248,8 +248,8 @@ struct FCameraRigAssetReference : public FBaseCameraObjectReference
 
 public:
 
-	FCameraRigAssetReference();
-	FCameraRigAssetReference(UCameraRigAsset* InCameraRig);
+	GAMEPLAYCAMERAS_API FCameraRigAssetReference();
+	GAMEPLAYCAMERAS_API FCameraRigAssetReference(UCameraRigAsset* InCameraRig);
 
 	/** Returns whether this reference points to a valid camera rig. */
 	bool IsValid() const
