Cameras: add Blueprint function for checking the validity of a shake instance ID

diff --git a/Engine/Plugins/Cameras/GameplayCameras/Source/GameplayCameras/Public/Core/CameraShakeInstanceID.h b/Engine/Plugins/Cameras/GameplayCameras/Source/GameplayCameras/Public/Core/CameraShakeInstanceID.h
--- a/Engine/Plugins/Cameras/GameplayCameras/Source/GameplayCameras/Public/Core/CameraShakeInstanceID.h
+++ b/Engine/Plugins/Cameras/GameplayCameras/Source/GameplayCameras/Public/Core/CameraShakeInstanceID.h
@@ -3,6 +3,7 @@
 #pragma once
 
 #include "CoreTypes.h"
+#include "Kismet/BlueprintFunctionLibrary.h"
 
 #include "CameraShakeInstanceID.generated.h"
 
@@ -57,3 +58,25 @@ private:
 	uint32 Value;
 };
 
+/**
+ * Blueprint function library for camera shake instance IDs.
+ */
+UCLASS()
+class UCameraShakeInstanceFunctionLibrary : public UBlueprintFunctionLibrary
+{
+	GENERATED_BODY()
+
+public:
+
+	/**
+	 * Returns whether a camera shake instance ID is valid. 
+	 *
+	 * Note that this doesn't check if the associated camera shake is still running.
+	 */
+	UFUNCTION(BlueprintPure, Category="Camera Shake", meta=(DisplayName="Is Valid Camera Shake Instance ID"))
+	static bool IsValid(const FCameraShakeInstanceID CameraShakeInstanceID)
+	{
+		return CameraShakeInstanceID.IsValid();
+	}
+};
+
