Cameras: clean-up how components and player managers interact

Simplify the camera component:
- It now has an option to either run its own camera system ("standalone" mode, which is the default) or only run a camera evalution context.
- Remove some of the old code that attempted to handle a mix of both cases (such as with the non-immediate deactivation code)

Simply the player camera manager:
- Now anything that is set as the view target is run as a "blackbox", whether it's an arbitrary actor or vanilla camera actor or procedural camera (in which case it's better to use the aforementioned "standalone" mode)
- To use fancier features such as blending or children directors, the procedural camera must be activated using dedicated new methods such as "Activate Gameplay Camera".

Other clean-ups:
- Move and improve on all the methods to activate/deactivate rigs and rig modifiers, and start/stop shakes.

diff --git a/Engine/Plugins/Cameras/GameplayCameras/Source/GameplayCameras/Private/Core/CameraEvaluationContextStack.cpp b/Engine/Plugins/Cameras/GameplayCameras/Source/GameplayCameras/Private/Core/CameraEvaluationContextStack.cpp
--- a/Engine/Plugins/Cameras/GameplayCameras/Source/GameplayCameras/Private/Core/CameraEvaluationContextStack.cpp
+++ b/Engine/Plugins/Cameras/GameplayCameras/Source/GameplayCameras/Private/Core/CameraEvaluationContextStack.cpp
@@ -177,6 +177,21 @@ void FCameraEvaluationContextStack::GetAllContexts(TArray<TSharedPtr<FCameraEval
 	}
 }
 
+TSharedPtr<FCameraEvaluationContext> FCameraEvaluationContextStack::FindContextByOwner(const UObject* InOwner) const
+{
+	for (const FContextEntry& Entry : Entries)
+	{
+		if (TSharedPtr<FCameraEvaluationContext> Context = Entry.WeakContext.Pin())
+		{
+			if (Context->GetOwner() == InOwner)
+			{
+				return Context;
+			}
+		}
+	}
+	return TSharedPtr<FCameraEvaluationContext>();
+}
+
 void FCameraEvaluationContextStack::Reset()
 {
 	for (FContextEntry& Entry : Entries)
diff --git a/Engine/Plugins/Cameras/GameplayCameras/Source/GameplayCameras/Private/GameFramework/GameplayCameraComponent.cpp b/Engine/Plugins/Cameras/GameplayCameras/Source/GameplayCameras/Private/GameFramework/GameplayCameraComponent.cpp
--- a/Engine/Plugins/Cameras/GameplayCameras/Source/GameplayCameras/Private/GameFramework/GameplayCameraComponent.cpp
+++ b/Engine/Plugins/Cameras/GameplayCameras/Source/GameplayCameras/Private/GameFramework/GameplayCameraComponent.cpp
@@ -83,7 +83,7 @@ void UGameplayCameraComponent::NotifyChangeCameraReference()
 {
 	using namespace UE::Cameras;
 
-	if (HasCameraEvaluationContext())
+	if (HasEvaluationContext())
 	{
 		// Sequencer animated some of our parameters... look for those whose value changed, compared to our
 		// cached parameter bag, and re-apply them to the evaluation context.
@@ -109,11 +109,11 @@ void UGameplayCameraComponent::OnCameraAssetBuilt(const UCameraAsset* InCameraAs
 	// If our camera asset was just built, it may have some new parameters. We need to rebuild
 	// our variable table and context data table, and re-apply overrides.
 	CameraReference.RebuildParametersIfNeeded();
-	if (HasCameraEvaluationContext())
+	if (HasEvaluationContext())
 	{
 		const FCameraAssetAllocationInfo& AllocationInfo = InCameraAsset->GetAllocationInfo();
-		ReinitializeCameraEvaluationContext(AllocationInfo.VariableTableInfo, AllocationInfo.ContextDataTableInfo);
-		UpdateCameraEvaluationContext(true);
+		ReinitializeEvaluationContext(AllocationInfo.VariableTableInfo, AllocationInfo.ContextDataTableInfo);
+		UpdateEvaluationContext(true);
 	}
 }
 
@@ -126,17 +126,17 @@ void UGameplayCameraComponent::PostEditChangeProperty( struct FPropertyChangedEv
 	const FName MemberPropertyName = PropertyChangedEvent.GetMemberPropertyName();
 	if (MemberPropertyName == GET_MEMBER_NAME_CHECKED(UGameplayCameraComponent, CameraReference))
 	{
-		if (HasCameraEvaluationContext())
+		if (HasEvaluationContext())
 		{
 			if (PropertyChangedEvent.GetPropertyName() == TEXT("CameraAsset"))
 			{
 				// The camera asset has changed! Recreate the context.
-				RecreateEditorWorldCameraEvaluationContext();
+				RecreateEditorWorldEvaluationContext();
 			}
 			else
 			{
 				// Otherwise, maybe one of the parameter overrides has changed. Re-apply them.
-				UpdateCameraEvaluationContext(true);
+				UpdateEvaluationContext(true);
 			}
 		}
 	}
diff --git a/Engine/Plugins/Cameras/GameplayCameras/Source/GameplayCameras/Private/GameFramework/GameplayCameraComponentBase.cpp b/Engine/Plugins/Cameras/GameplayCameras/Source/GameplayCameras/Private/GameFramework/GameplayCameraComponentBase.cpp
--- a/Engine/Plugins/Cameras/GameplayCameras/Source/GameplayCameras/Private/GameFramework/GameplayCameraComponentBase.cpp
+++ b/Engine/Plugins/Cameras/GameplayCameras/Source/GameplayCameras/Private/GameFramework/GameplayCameraComponentBase.cpp
@@ -20,8 +20,6 @@
 #include "Misc/AssertionMacros.h"
 #include "Misc/EngineVersionComparison.h"
 #include "SceneView.h"
-#include "Services/CameraModifierService.h"
-#include "Services/CameraShakeService.h"
 #include "ShowFlags.h"
 
 #if UE_VERSION_NEWER_THAN_OR_EQUAL(5,6,0)
@@ -55,8 +53,8 @@ UGameplayCameraComponentBase::UGameplayCameraComponentBase(const FObjectInitiali
 
 void UGameplayCameraComponentBase::BeginDestroy()
 {
-	DestroyCameraSystem();
-	DestroyCameraEvaluationContext();
+	DestroyCameraSystemHost();
+	DestroyEvaluationContext();
 
 	Super::BeginDestroy();
 }
@@ -83,15 +81,13 @@ TSharedPtr<UE::Cameras::FCameraEvaluationContext> UGameplayCameraComponentBase::
 	return EvaluationContext;
 }
 
-void UGameplayCameraComponentBase::ActivateCameraForPlayerIndex(
-		int32 PlayerIndex, bool bSetAsViewTarget, EGameplayCameraComponentActivationMode ActivationMode)
+void UGameplayCameraComponentBase::ActivateCameraForPlayerIndex(int32 PlayerIndex, bool bSetAsViewTarget)
 {
 	APlayerController* PlayerController = UGameplayStatics::GetPlayerController(this, PlayerIndex);
-	ActivateCameraForPlayerController(PlayerController, bSetAsViewTarget, ActivationMode);
+	ActivateCameraForPlayerController(PlayerController, bSetAsViewTarget);
 }
 
-void UGameplayCameraComponentBase::ActivateCameraForPlayerController(
-		APlayerController* PlayerController, bool bSetAsViewTarget, EGameplayCameraComponentActivationMode ActivationMode)
+void UGameplayCameraComponentBase::ActivateCameraForPlayerController(APlayerController* PlayerController, bool bSetAsViewTarget)
 {
 	using namespace UE::Cameras;
 
@@ -99,118 +95,24 @@ void UGameplayCameraComponentBase::ActivateCameraForPlayerController(
 	// our private camera system.
 	Super::Activate(false);
 
-	// Deactivate any existing evaluation context immediately first, since we might be re-activating with
-	// a different player controller, a different insertion point in the director tree, etc.
-	DeactivateCameraEvaluationContext(true);
-
-	// See if we are going to run the camera system on our own, or if we can activate our camera inside
-	// the player camera manager.
-	IGameplayCameraSystemHost* PlayerControllerHost = nullptr;
-	if (bSetAsViewTarget && PlayerController)
-	{
-		PlayerControllerHost = Cast<IGameplayCameraSystemHost>(PlayerController->PlayerCameraManager);
-	}
-	if (!PlayerControllerHost && ActivationMode != EGameplayCameraComponentActivationMode::Push)
-	{
-		FFrame::KismetExecutionMessage(
-				*FString::Format(
-					TEXT("Gameplay Camera component '{0}.{1}' cannot activate with mode '{2}' because no camera system "
-						 "was found on the given player controller, or no player controller was specified."),
-					{ *GetNameSafe(GetOwner()), *GetNameSafe(this), *UEnum::GetValueAsString(ActivationMode) }),
-				ELogVerbosity::Warning);
-	}
+	EnsureEvaluationContext(PlayerController);
+	EnsureCameraSystemHostIfNeeded();
 
-	if (PlayerControllerHost)
+	if (bSetAsViewTarget)
 	{
-		DestroyCameraSystem();
-		ActivateCameraEvaluationContext(PlayerController, PlayerControllerHost, ActivationMode);
-	}
-	else
-	{
-		EnsureCameraSystemHost();
-		ActivateCameraEvaluationContext(PlayerController, this, EGameplayCameraComponentActivationMode::Push);
-
-		if (bSetAsViewTarget && PlayerController)
+		if (PlayerController)
 		{
 			AActor* OwnerActor = GetOwner();
 			PlayerController->SetViewTarget(OwnerActor);
 		}
-	}
-}
-
-bool UGameplayCameraComponentBase::CalcCamera(float DeltaTime, FMinimalViewInfo& OutResult)
-{
-	using namespace UE::Cameras;
-
-	if (CameraSystemEvaluator)
-	{
-		CameraSystemEvaluator->GetEvaluatedCameraView(OutResult);
-		return true;
-	}
-	return false;
-}
-
-void UGameplayCameraComponentBase::DeactivateCamera(bool bImmediately)
-{
-	DeactivateCameraEvaluationContext(bImmediately);
-}
-
-void UGameplayCameraComponentBase::DeactivateCameraEvaluationContext(bool bImmediately)
-{
-	using namespace UE::Cameras;
-
-	if (!bImmediately && HasCameraSystem())
-	{
-		FFrame::KismetExecutionMessage(
-				*FString::Format(
-					TEXT("Can't deactivate Gameplay Camera component '{0}.{1}' with bImmediately set to false: "
-						 "this component is running its own camera system and cannot therefore blend out from itself."),
-					{ *GetNameSafe(GetOwner()), *GetNameSafe(this) }),
-				ELogVerbosity::Warning);
-		bImmediately = true;
-	}
-
-	if (EvaluationContext && EvaluationContext->IsActive())
-	{
-		UE_LOG(LogCameraSystem, Log, TEXT("Deactivating gameplay camera '%s'."), *GetNameSafe(this));
-
-		TSharedPtr<FCameraSystemEvaluator> HostEvaluator = EvaluationContext->GetCameraSystemEvaluator();
-		ensure(HostEvaluator);
-
-		if (TSharedPtr<FCameraEvaluationContext> ParentContext = EvaluationContext->GetParentContext())
-		{
-			ParentContext->RemoveChildContext(EvaluationContext.ToSharedRef());
-		}
 		else
 		{
-			HostEvaluator->RemoveEvaluationContext(EvaluationContext.ToSharedRef());
-		}
-
-		if (bImmediately)
-		{
-			// We are deactivating immediately (i.e. without letting our camera rigs blend out), so make 
-			// sure everything is frozen or disabled before we delete our evaluation context.
-			FRootCameraNodeEvaluator* RootNodeEvaluator = HostEvaluator->GetRootNodeEvaluator();
-			RootNodeEvaluator->DeactivateAllCameraRigs(EvaluationContext, true);
-		}
-		else
-		{
-			// Don't deactivate the component right away: we still need to update our evaluation context 
-			// while any running camera rigs blend out.
-			PendingDeactivateCameraSystemEvaluator = HostEvaluator;
-			bIsPendingDeactivate = true;
-		}
-	}
-
-	if (bImmediately)
-	{
-		if (OutputCameraComponent)
-		{
-			OutputCameraComponent->SetRelativeTransform(FTransform());
+			FFrame::KismetExecutionMessage(
+				*FString::Format(
+					TEXT("Gameplay Camera component '{0}.{1}' cannot set itself as view target because no player controller was provided."),
+					{ *GetNameSafe(GetOwner()), *GetNameSafe(this) }),
+				ELogVerbosity::Error);
 		}
-
-		DestroyCameraSystem();
-		DestroyCameraEvaluationContext();
 	}
 }
 
@@ -231,11 +133,22 @@ bool UGameplayCameraComponentBase::CanRunCameraSystem() const
 #endif  // WITH_EDITOR
 }
 
-void UGameplayCameraComponentBase::EnsureCameraSystemHost()
+bool UGameplayCameraComponentBase::EnsureCameraSystemHostIfNeeded()
 {
 	using namespace UE::Cameras;
 
-	if (!HasCameraSystem() && CanRunCameraSystem())
+	if (!bRunStandaloneCameraSystem)
+	{
+		return false;
+	}
+
+	if (!CanRunCameraSystem())
+	{
+		return false;
+	}
+
+	// Create a hosted camera system, and then create and active the evalution context.
+	if (!HasCameraSystem())
 	{
 		AActor* OwnerActor = GetOwner();
 		UE_LOG(LogCameraSystem, Log, 
@@ -252,85 +165,66 @@ void UGameplayCameraComponentBase::EnsureCameraSystemHost()
 #endif  // WITH_EDITOR
 		InitializeCameraSystem(Params);
 	}
-}
 
-void UGameplayCameraComponentBase::ActivateCameraEvaluationContext(APlayerController* PlayerController, IGameplayCameraSystemHost* Host, EGameplayCameraComponentActivationMode ActivationMode)
-{
-	using namespace UE::Cameras;
-
-	if (!ensure(Host))
+	if (!ensureMsgf(
+				EvaluationContext.IsValid(),
+				TEXT("Can't activate Gameplay Camera component '{0}.{1}': failed to create evaluation context!"),
+				*GetNameSafe(GetOwner()), *GetNameSafe(this)))
 	{
-		return;
+		return false;
 	}
 
-	if (!CanRunCameraSystem())
+	if (!ensureMsgf(
+				!EvaluationContext->IsActive(),
+				TEXT("Can't activate Gameplay Camera component '{0}.{1}': it is already active!"),
+				*GetNameSafe(GetOwner()), *GetNameSafe(this)))
 	{
-		return;
+		return false;
 	}
 
-	CreateCameraEvaluationContext(PlayerController);
+	UE_LOG(LogCameraSystem, Log, 
+			TEXT("Activating gameplay camera '%s.%s' with its hosted camera system."),
+			*GetNameSafe(GetOwner()), *GetNameSafe(this));
+
+	TSharedPtr<FCameraSystemEvaluator> HostedEvaluator = GetCameraSystemEvaluator();
+	check(HostedEvaluator.IsValid());
+	FCameraEvaluationContextStack& ContextStack = HostedEvaluator->GetEvaluationContextStack();
+	ContextStack.PushContext(EvaluationContext.ToSharedRef());
 
-	if (EvaluationContext->IsActive())
+	return true;
+}
+
+void UGameplayCameraComponentBase::DestroyCameraSystemHost()
+{
+	// Reset the output camera component and destroy the hosted camera system.
+	if (OutputCameraComponent)
 	{
-		FFrame::KismetExecutionMessage(
-				*FString::Format(
-					TEXT("Can't activate Gameplay Camera component '{0}.{1}': it is already active!"),
-					{ *GetNameSafe(GetOwner()), *GetNameSafe(this) }),
-				ELogVerbosity::Error);
-		return;
+		OutputCameraComponent->SetRelativeTransform(FTransform());
 	}
 
-	UE_LOG(LogCameraSystem, Log, 
-			TEXT("Activating gameplay camera '%s' with mode '%s'."),
-			*GetNameSafe(this), 
-			*UEnum::GetValueAsString(ActivationMode));
-
-	TSharedPtr<FCameraSystemEvaluator> HostEvaluator = Host->GetCameraSystemEvaluator();
-	FCameraEvaluationContextStack& ContextStack = HostEvaluator->GetEvaluationContextStack();
-
-	switch (ActivationMode)
-	{
-		case EGameplayCameraComponentActivationMode::Push:
-			ContextStack.PushContext(EvaluationContext.ToSharedRef());
-			break;
-		case EGameplayCameraComponentActivationMode::PushAndInsert:
-			{
-				TSharedPtr<FCameraEvaluationContext> PreviousActiveContext = ContextStack.GetActiveContext();
-				ContextStack.PushContext(EvaluationContext.ToSharedRef());
-				if (PreviousActiveContext)
-				{
-					ContextStack.RemoveContext(PreviousActiveContext.ToSharedRef());
-					EvaluationContext->AddChildContext(PreviousActiveContext.ToSharedRef());
-				}
-			}
-			break;
-		case EGameplayCameraComponentActivationMode::InsertOrPush:
-			{
-				if (TSharedPtr<FCameraEvaluationContext> ActiveContext = ContextStack.GetActiveContext())
-				{
-					ActiveContext->AddChildContext(EvaluationContext.ToSharedRef());
-				}
-				else
-				{
-					ContextStack.PushContext(EvaluationContext.ToSharedRef());
-				}
-			}
-			break;
-		default:
-			ensure(false);
-			break;
-	}
-
-	// Cancel any ongoing deactivation.
-	PendingDeactivateCameraSystemEvaluator = nullptr;
-	bIsPendingDeactivate = false;
-}
-
-void UGameplayCameraComponentBase::CreateCameraEvaluationContext(APlayerController* PlayerController)
+	DestroyCameraSystem();
+}
+
+void UGameplayCameraComponentBase::EnsureEvaluationContext(APlayerController* PlayerController)
 {
 	using namespace UE::Cameras;
 
-	ensure(CanRunCameraSystem());
+	if (EvaluationContext.IsValid() && EvaluationContext->GetPlayerController() != PlayerController)
+	{
+		FFrame::KismetExecutionMessage(
+				*FString::Format(
+					TEXT("Trying to create an evaluation context for Player Controller '{0}' on Gameplay Camera component '{1}.{2}' "
+						 "but it was already created for Player Controller '{1}'. The evaluation context will be re-created."),
+					{ 
+						*GetNameSafe(PlayerController), 
+						*GetNameSafe(GetOwner()),
+						*GetNameSafe(this),
+						*GetNameSafe(EvaluationContext->GetPlayerController())
+					}),
+				ELogVerbosity::Warning);
+
+		DestroyEvaluationContext();
+	}
 
 	if (!EvaluationContext.IsValid())
 	{
@@ -350,7 +244,7 @@ void UGameplayCameraComponentBase::CreateCameraEvaluationContext(APlayerControll
 
 			UCameraAsset* PlaceholderCameraAsset = NewObject<UCameraAsset>();
 			PlaceholderCameraAsset->SetCameraDirector(PlaceholderCameraDirector);
-			
+
 			CameraAsset = PlaceholderCameraAsset;
 		}
 
@@ -362,156 +256,177 @@ void UGameplayCameraComponentBase::CreateCameraEvaluationContext(APlayerControll
 		InitParams.PlayerController = PlayerController;
 		EvaluationContext->Initialize(InitParams);
 
-		UpdateCameraEvaluationContext(true);
+		const bool bForceApplyParameterOverrides = true;
+		UpdateEvaluationContext(bForceApplyParameterOverrides);
 	}
 
 	ensure(EvaluationContext.IsValid());
 }
 
-void UGameplayCameraComponentBase::DestroyCameraEvaluationContext()
+void UGameplayCameraComponentBase::DestroyEvaluationContext()
 {
 	EvaluationContext = nullptr;
 }
 
-#define UE_PRIVATE_GAMEPLAY_CAMERA_COMPONENT_VALIDATE_EVALUATION_CONTEXT(ErrorMsg, ErrorResult)\
-	using namespace UE::Cameras;\
-	if (!EvaluationContext)\
-	{\
-		FFrame::KismetExecutionMessage(\
-				*FString::Format(\
-					TEXT(#ErrorMsg " on Gameplay Camera component '{0}.{1}': it isn't active."),\
-					{ *GetNameSafe(GetOwner()), *GetNameSafe(this) }),\
-				ELogVerbosity::Error);\
-		return ErrorResult;\
-	}
+#define UE_PRIVATE_GAMEPLAY_CAMERA_COMPONENT_ERROR_MESSAGE(ErrorMsg)\
+	FFrame::KismetExecutionMessage(\
+			*FString::Format(\
+				TEXT(#ErrorMsg " on Gameplay Camera component '{0}.{1}': it isn't active."),\
+				{ *GetNameSafe(GetOwner()), *GetNameSafe(this) }),\
+			ELogVerbosity::Error);\
 
 FBlueprintCameraEvaluationDataRef UGameplayCameraComponentBase::GetInitialResult() const
 {
-	UE_PRIVATE_GAMEPLAY_CAMERA_COMPONENT_VALIDATE_EVALUATION_CONTEXT("Can't get shared camera data", FBlueprintCameraEvaluationDataRef());
-
-	return FBlueprintCameraEvaluationDataRef::MakeExternalRef(&EvaluationContext->GetInitialResult());
+	if (EvaluationContext)
+	{
+		return FBlueprintCameraEvaluationDataRef::MakeExternalRef(&EvaluationContext->GetInitialResult());
+	}
+	else
+	{
+		UE_PRIVATE_GAMEPLAY_CAMERA_COMPONENT_ERROR_MESSAGE("Can't get shared camera data");
+		return FBlueprintCameraEvaluationDataRef();
+	}
 }
 
 FBlueprintCameraEvaluationDataRef UGameplayCameraComponentBase::GetConditionalResult(ECameraEvaluationDataCondition Condition) const
 {
-	UE_PRIVATE_GAMEPLAY_CAMERA_COMPONENT_VALIDATE_EVALUATION_CONTEXT("Can't get conditional camera data", FBlueprintCameraEvaluationDataRef());
-
-	return FBlueprintCameraEvaluationDataRef::MakeExternalRef(&EvaluationContext->GetOrAddConditionalResult(Condition));
+	if (EvaluationContext)
+	{
+		return FBlueprintCameraEvaluationDataRef::MakeExternalRef(&EvaluationContext->GetOrAddConditionalResult(Condition));
+	}
+	else
+	{
+		UE_PRIVATE_GAMEPLAY_CAMERA_COMPONENT_ERROR_MESSAGE("Can't get conditional camera data");
+		return  FBlueprintCameraEvaluationDataRef();
+	}
 }
 
-#undef UE_PRIVATE_GAMEPLAY_CAMERA_COMPONENT_VALIDATE_EVALUATION_CONTEXT
-
-#define UE_PRIVATE_GAMEPLAY_CAMERA_COMPONENT_ACTIVATE_NON_MAIN_RIG(ErrorMsg, RigLayer)\
-	if (EvaluationContext.IsValid())\
-	{\
-		IGameplayCameraSystemHost::ActivateCameraRig(CameraRig, EvaluationContext, ECameraRigLayer::RigLayer);\
-	}\
-	else\
-	{\
-		FFrame::KismetExecutionMessage(\
-				*FString::Format(\
-					TEXT(#ErrorMsg " on Gameplay Camera component '{0}.{1}': it isn't active."),\
-					{ *GetNameSafe(GetOwner()), *GetNameSafe(this) }),\
-				ELogVerbosity::Error);\
+FCameraRigInstanceID UGameplayCameraComponentBase::ActivatePersistentBaseCameraRig(UCameraRigAsset* CameraRig)
+{
+	if (HasCameraSystem() && ensure(EvaluationContext))
+	{
+		return IGameplayCameraSystemHost::ActivateCameraRig(CameraRig, EvaluationContext, ECameraRigLayer::Base);
 	}
+	else
+	{
+		UE_PRIVATE_GAMEPLAY_CAMERA_COMPONENT_ERROR_MESSAGE("Can't activate base camera rig");
+		return FCameraRigInstanceID();
+	}
+}
 
-void UGameplayCameraComponentBase::ActivatePersistentBaseCameraRig(UCameraRigAsset* CameraRig)
+FCameraRigInstanceID UGameplayCameraComponentBase::ActivatePersistentGlobalCameraRig(UCameraRigAsset* CameraRig)
 {
-	UE_PRIVATE_GAMEPLAY_CAMERA_COMPONENT_ACTIVATE_NON_MAIN_RIG("Can't activate base camera rig", Base);
+	if (HasCameraSystem() && ensure(EvaluationContext))
+	{
+		return IGameplayCameraSystemHost::ActivateCameraRig(CameraRig, EvaluationContext, ECameraRigLayer::Global);
+	}
+	else
+	{
+		UE_PRIVATE_GAMEPLAY_CAMERA_COMPONENT_ERROR_MESSAGE("Can't activate global camera rig");
+		return FCameraRigInstanceID();
+	}
 }
 
-void UGameplayCameraComponentBase::ActivatePersistentGlobalCameraRig(UCameraRigAsset* CameraRig)
+FCameraRigInstanceID UGameplayCameraComponentBase::ActivatePersistentVisualCameraRig(UCameraRigAsset* CameraRig)
 {
-	UE_PRIVATE_GAMEPLAY_CAMERA_COMPONENT_ACTIVATE_NON_MAIN_RIG("Can't activate global camera rig", Global);
+	if (HasCameraSystem() && ensure(EvaluationContext))
+	{
+		return IGameplayCameraSystemHost::ActivateCameraRig(CameraRig, EvaluationContext, ECameraRigLayer::Visual);
+	}
+	else
+	{
+		UE_PRIVATE_GAMEPLAY_CAMERA_COMPONENT_ERROR_MESSAGE("Can't activate visual camera rig");
+		return FCameraRigInstanceID();
+	}
 }
 
-void UGameplayCameraComponentBase::ActivatePersistentVisualCameraRig(UCameraRigAsset* CameraRig)
+void UGameplayCameraComponentBase::DeactivateCameraRig(FCameraRigInstanceID InstanceID, bool bImmediately)
 {
-	UE_PRIVATE_GAMEPLAY_CAMERA_COMPONENT_ACTIVATE_NON_MAIN_RIG("Can't activate visual camera rig", Visual);
+	if (HasCameraSystem())
+	{
+		IGameplayCameraSystemHost::DeactivateCameraRig(InstanceID, bImmediately);
+	}
+	else
+	{
+		UE_PRIVATE_GAMEPLAY_CAMERA_COMPONENT_ERROR_MESSAGE("Can't deactivate camera rig");
+	}
 }
 
-#undef UE_PRIVATE_GAMEPLAY_CAMERA_COMPONENT_ACTIVATE_NON_MAIN_RIG
-
 FCameraRigInstanceID UGameplayCameraComponentBase::StartGlobalCameraModifierRig(const UCameraRigAsset* CameraRig, int32 OrderKey)
 {
-	using namespace UE::Cameras;
-
-	if (CameraSystemEvaluator)
+	if (HasCameraSystem() && ensure(EvaluationContext))
 	{
-		TSharedPtr<FCameraModifierService> CameraModifierService = CameraSystemEvaluator->FindEvaluationService<FCameraModifierService>();
-		return CameraModifierService->StartCameraModifierRig(CameraRig, EvaluationContext.ToSharedRef(), ECameraRigLayer::Global, OrderKey);
+		return IGameplayCameraSystemHost::StartCameraModifierRig(CameraRig, EvaluationContext, ECameraRigLayer::Global, OrderKey);
+	}
+	else
+	{
+		UE_PRIVATE_GAMEPLAY_CAMERA_COMPONENT_ERROR_MESSAGE("Can't start global camera modifier rig");
+		return FCameraRigInstanceID();
 	}
-
-	return FCameraRigInstanceID();
 }
 
 FCameraRigInstanceID UGameplayCameraComponentBase::StartVisualCameraModifierRig(const UCameraRigAsset* CameraRig, int32 OrderKey)
 {
-	using namespace UE::Cameras;
-
-	if (CameraSystemEvaluator)
+	if (HasCameraSystem() && ensure(EvaluationContext))
 	{
-		TSharedPtr<FCameraModifierService> CameraModifierService = CameraSystemEvaluator->FindEvaluationService<FCameraModifierService>();
-		return CameraModifierService->StartCameraModifierRig(CameraRig, EvaluationContext.ToSharedRef(), ECameraRigLayer::Visual, OrderKey);
+		return IGameplayCameraSystemHost::StartCameraModifierRig(CameraRig, EvaluationContext, ECameraRigLayer::Visual, OrderKey);
+	}
+	else
+	{
+		UE_PRIVATE_GAMEPLAY_CAMERA_COMPONENT_ERROR_MESSAGE("Can't start visual camera modifier rig");
+		return FCameraRigInstanceID();
 	}
-
-	return FCameraRigInstanceID();
 }
 
 void UGameplayCameraComponentBase::StopCameraModifierRig(FCameraRigInstanceID InstanceID, bool bImmediately)
 {
-	using namespace UE::Cameras;
-
-	if (CameraSystemEvaluator)
+	if (HasCameraSystem())
 	{
-		TSharedPtr<FCameraModifierService> CameraModifierService = CameraSystemEvaluator->FindEvaluationService<FCameraModifierService>();
-		return CameraModifierService->StopCameraModifierRig(InstanceID, bImmediately);
+		IGameplayCameraSystemHost::StopCameraModifierRig(InstanceID, bImmediately);
+	}
+	else
+	{
+		UE_PRIVATE_GAMEPLAY_CAMERA_COMPONENT_ERROR_MESSAGE("Can't stop camera modifier rig");
 	}
 }
 
 FCameraShakeInstanceID UGameplayCameraComponentBase::StartCameraShakeAsset(const UCameraShakeAsset* CameraShake, float ShakeScale, ECameraShakePlaySpace PlaySpace, FRotator UserPlaySpaceRotation)
 {
-	using namespace UE::Cameras;
-
-	if (CameraSystemEvaluator)
+	if (HasCameraSystem())
 	{
-		FStartCameraShakeParams Params;
-		Params.CameraShake = CameraShake;
-		Params.ShakeScale = ShakeScale;
-		Params.PlaySpace = PlaySpace;
-		Params.UserPlaySpaceRotation = UserPlaySpaceRotation;
-
-		TSharedPtr<FCameraShakeService> CameraShakeService = CameraSystemEvaluator->FindEvaluationService<FCameraShakeService>();
-		return CameraShakeService->StartCameraShake(Params);
+		return IGameplayCameraSystemHost::StartCameraShake(CameraShake, ShakeScale, PlaySpace, UserPlaySpaceRotation);
+	}
+	else
+	{
+		UE_PRIVATE_GAMEPLAY_CAMERA_COMPONENT_ERROR_MESSAGE("Can't start camera shake");
+		return FCameraShakeInstanceID();
 	}
-
-	return FCameraShakeInstanceID();
 }
 
 bool UGameplayCameraComponentBase::IsCameraShakeAssetPlaying(FCameraShakeInstanceID InInstanceID) const
 {
-	using namespace UE::Cameras;
-
-	if (CameraSystemEvaluator)
+	if (HasCameraSystem())
 	{
-		TSharedPtr<FCameraShakeService> CameraShakeService = CameraSystemEvaluator->FindEvaluationService<FCameraShakeService>();
-		return CameraShakeService->IsCameraShakePlaying(InInstanceID);
+		return IGameplayCameraSystemHost::IsCameraShakePlaying(InInstanceID);
 	}
 	return false;
 }
 
 bool UGameplayCameraComponentBase::StopCameraShakeAsset(FCameraShakeInstanceID InInstanceID, bool bImmediately)
 {
-	using namespace UE::Cameras;
-
-	if (CameraSystemEvaluator)
+	if (HasCameraSystem())
 	{
-		TSharedPtr<FCameraShakeService> CameraShakeService = CameraSystemEvaluator->FindEvaluationService<FCameraShakeService>();
-		return CameraShakeService->StopCameraShake(InInstanceID, bImmediately);
+		return IGameplayCameraSystemHost::StopCameraShake(InInstanceID, bImmediately);
+	}
+	else
+	{
+		UE_PRIVATE_GAMEPLAY_CAMERA_COMPONENT_ERROR_MESSAGE("Can't stop camera shake");
+		return false;
 	}
-	return false;
 }
 
+#undef UE_PRIVATE_GAMEPLAY_CAMERA_COMPONENT_VALIDATE_EVALUATION_CONTEXT
+
 FRotator UGameplayCameraComponentBase::GetEvaluatedCameraRotation() const
 {
 	using namespace UE::Cameras;
@@ -562,14 +477,15 @@ void UGameplayCameraComponentBase::BeginPlay()
 		}
 		else
 		{
-			ActivateCameraForPlayerController(nullptr, false, EGameplayCameraComponentActivationMode::Push);
+			ActivateCameraForPlayerController(nullptr, false);
 		}
 	}
 }
 
 void UGameplayCameraComponentBase::EndPlay(const EEndPlayReason::Type EndPlayReason)
 {
-	DeactivateCameraEvaluationContext(true);
+	DestroyCameraSystemHost();
+	DestroyEvaluationContext();
 
 	Super::EndPlay(EndPlayReason);
 }
@@ -584,16 +500,15 @@ void UGameplayCameraComponentBase::Activate(bool bReset)
 
 	if (bDoActivate)
 	{
-		DeactivateCameraEvaluationContext(true);
-
-		EnsureCameraSystemHost();
-		ActivateCameraEvaluationContext(nullptr, this, EGameplayCameraComponentActivationMode::Push);
+		EnsureEvaluationContext(nullptr);
+		EnsureCameraSystemHostIfNeeded();
 	}
 }
 
 void UGameplayCameraComponentBase::Deactivate()
 {
-	DeactivateCameraEvaluationContext(true);
+	DestroyCameraSystemHost();
+	DestroyEvaluationContext();
 
 	Super::Deactivate();
 }
@@ -611,10 +526,7 @@ void UGameplayCameraComponentBase::TickComponent(float DeltaTime, ELevelTick Tic
 
 #endif  // WITH_EDITOR
 
-	if (EvaluationContext)
-	{
-		UpdateCameraEvaluationContext(false);
-	}
+	UpdateEvaluationContext(false);
 
 #if WITH_EDITOR
 
@@ -622,7 +534,7 @@ void UGameplayCameraComponentBase::TickComponent(float DeltaTime, ELevelTick Tic
 	{
 		UpdateCameraSystemForEditorPreview(DeltaTime);
 	}
-	else
+	else if (bRunStandaloneCameraSystem)
 	{
 		UpdateCameraSystem(DeltaTime);
 		UpdateControlRotationIfNeeded();
@@ -630,81 +542,51 @@ void UGameplayCameraComponentBase::TickComponent(float DeltaTime, ELevelTick Tic
 
 #else
 
-	UpdateCameraSystem(DeltaTime);
-	UpdateControlRotationIfNeeded();
+	if (bRunStandaloneCameraSystem)
+	{
+		UpdateCameraSystem(DeltaTime);
+		UpdateControlRotationIfNeeded();
+	}
 
 #endif  // WITH_EDITOR
 
 	UpdateOutputCameraComponent();
-
-	CheckPendingDeactivation();
 }
 
-void UGameplayCameraComponentBase::CheckPendingDeactivation()
+void UGameplayCameraComponentBase::UpdateEvaluationContext(bool bForceApplyParameterOverrides)
 {
 	using namespace UE::Cameras;
 
-	if (!bIsPendingDeactivate)
-	{
-		return;
-	}
-
-	// If we were deactivating, we wait until we don't have any running camera rig anymore,
-	// at which point we can tear down all our evaluation apparatus.
-	bool bDoneDeactivating = true;
-
-	TSharedPtr<FCameraSystemEvaluator> CameraSystemToCheck = PendingDeactivateCameraSystemEvaluator.Pin();
-	if (EvaluationContext && CameraSystemToCheck)
-	{
-		FRootCameraNodeEvaluator* RootNodeEvaluator = CameraSystemToCheck->GetRootNodeEvaluator();
-		bDoneDeactivating = (RootNodeEvaluator->HasAnyRunningCameraRig(EvaluationContext) == false);
-	}
-
-	if (bDoneDeactivating)
+	if (EvaluationContext)
 	{
-		DestroyCameraSystem();
-		DestroyCameraEvaluationContext();
-
-		// Only call the base class method here: we just want to finish deactivating ourselves
-		// by stopping ticking.
-		Super::Deactivate();
-
-		PendingDeactivateCameraSystemEvaluator = nullptr;
-		bIsPendingDeactivate = false;
-	}
-}
-
-void UGameplayCameraComponentBase::UpdateCameraEvaluationContext(bool bForceApplyParameterOverrides)
-{
-	using namespace UE::Cameras;
-
-	FCameraNodeEvaluationResult& InitialResult = EvaluationContext->GetInitialResult();
+		FCameraNodeEvaluationResult& InitialResult = EvaluationContext->GetInitialResult();
 
-	const FTransform& OwnerTransform = GetComponentTransform();
-	InitialResult.CameraPose.SetTransform(OwnerTransform, true);
-	InitialResult.bIsCameraCut = false;
-	InitialResult.bIsValid = true;
+		const FTransform& OwnerTransform = GetComponentTransform();
+		InitialResult.CameraPose.SetTransform(OwnerTransform, true);
+		InitialResult.bIsCameraCut = false;
+		InitialResult.bIsValid = true;
 
-	if (bIsCameraCutNextFrame)
-	{
-		InitialResult.bIsCameraCut = true;
-		bIsCameraCutNextFrame = false;
-	}
+		if (bIsCameraCutNextFrame)
+		{
+			InitialResult.bIsCameraCut = true;
+			bIsCameraCutNextFrame = false;
+		}
 
-	OnUpdateCameraEvaluationContext(bForceApplyParameterOverrides);
+		OnUpdateCameraEvaluationContext(bForceApplyParameterOverrides);
 
 #if WITH_EDITOR
 
-	EvaluationContext->UpdateForEditorPreview();
+		EvaluationContext->UpdateForEditorPreview();
 
 #endif  // WITH_EDITOR
+	}
 }
 
 void UGameplayCameraComponentBase::UpdateControlRotationIfNeeded()
 {
 	using namespace UE::Cameras;
 
-	if (!HasCameraSystem() || !HasCameraEvaluationContext() || !bSetControlRotationWhenViewTarget)
+	if (!HasCameraSystem() || !HasEvaluationContext() || !bSetControlRotationWhenViewTarget)
 	{
 		return;
 	}
@@ -743,7 +625,7 @@ bool UGameplayCameraComponentBase::IsEditorWorld() const
 
 #if WITH_EDITOR
 
-void UGameplayCameraComponentBase::ReinitializeCameraEvaluationContext(
+void UGameplayCameraComponentBase::ReinitializeEvaluationContext(
 			const FCameraVariableTableAllocationInfo& VariableTableAllocationInfo,
 			const FCameraContextDataTableAllocationInfo& ContextDataTableAllocationInfo)
 {
@@ -765,7 +647,7 @@ void UGameplayCameraComponentBase::ReinitializeCameraEvaluationContext(
 	}
 }
 
-void UGameplayCameraComponentBase::RecreateEditorWorldCameraEvaluationContext()
+void UGameplayCameraComponentBase::RecreateEditorWorldEvaluationContext()
 {
 	using namespace UE::Cameras;
 
@@ -797,7 +679,7 @@ void UGameplayCameraComponentBase::RecreateEditorWorldCameraEvaluationContext()
 		}
 		EvaluationContext = nullptr;
 
-		CreateCameraEvaluationContext(nullptr);
+		EnsureEvaluationContext(nullptr);
 		if (CameraSystemEvaluator)
 		{
 			CameraSystemEvaluator->PushEvaluationContext(EvaluationContext.ToSharedRef());
@@ -877,7 +759,7 @@ void UGameplayCameraComponentBase::OnUpdateTransform(EUpdateTransformFlags Updat
 
 	if (bIsEditorWorld && EvaluationContext)
 	{
-		UpdateCameraEvaluationContext(false);
+		UpdateEvaluationContext(false);
 	}
 
 #endif  // WITH_EDITOR
@@ -899,10 +781,10 @@ void UGameplayCameraComponentBase::AutoManageEditorPreviewEvaluator()
 	{
 		// We want to run the camera logic in the editor but we haven't set things up for that.
 		// Let's create the preview evaluator and the evaluation context.
-		EnsureCameraSystemHost();
+		EnsureEvaluationContext(nullptr);
+		EnsureCameraSystemHostIfNeeded();
 
-		ActivateCameraEvaluationContext(nullptr, this, EGameplayCameraComponentActivationMode::Push);
-		if (EvaluationContext)
+		if (ensure(EvaluationContext))
 		{
 			EvaluationContext->SetEditorPreviewCameraRigIndex(EditorPreviewCameraRigIndex);
 		}
@@ -912,14 +794,8 @@ void UGameplayCameraComponentBase::AutoManageEditorPreviewEvaluator()
 	else if (!bCanRun && (CameraSystemEvaluator || EvaluationContext))
 	{
 		// We don't want to run the camera logic in the editor anymore. Let's tear things down.
-		DeactivateCameraEvaluationContext(true);
-		DestroyCameraSystem();
-		EvaluationContext = nullptr;
-
-		if (OutputCameraComponent)
-		{
-			OutputCameraComponent->SetRelativeTransform(FTransform());
-		}
+		DestroyCameraSystemHost();
+		DestroyEvaluationContext();
 	}
 }
 
diff --git a/Engine/Plugins/Cameras/GameplayCameras/Source/GameplayCameras/Private/GameFramework/GameplayCameraRigComponent.cpp b/Engine/Plugins/Cameras/GameplayCameras/Source/GameplayCameras/Private/GameFramework/GameplayCameraRigComponent.cpp
--- a/Engine/Plugins/Cameras/GameplayCameras/Source/GameplayCameras/Private/GameFramework/GameplayCameraRigComponent.cpp
+++ b/Engine/Plugins/Cameras/GameplayCameras/Source/GameplayCameras/Private/GameFramework/GameplayCameraRigComponent.cpp
@@ -96,15 +96,15 @@ void UGameplayCameraRigComponent::OnCameraRigAssetBuiltImpl()
 	using namespace UE::Cameras;
 
 	CameraRigReference.RebuildParametersIfNeeded();
-	if (HasCameraEvaluationContext())
+	if (HasEvaluationContext())
 	{
 #if WITH_EDITOR
 		const UCameraRigAsset* CameraRigAsset = CameraRigReference.GetCameraRig();
 		const FCameraObjectAllocationInfo& AllocationInfo = CameraRigAsset->AllocationInfo;
-		ReinitializeCameraEvaluationContext(AllocationInfo.VariableTableInfo, AllocationInfo.ContextDataTableInfo);
+		ReinitializeEvaluationContext(AllocationInfo.VariableTableInfo, AllocationInfo.ContextDataTableInfo);
 #endif  // WITH_EDITOR
 
-		UpdateCameraEvaluationContext(true);
+		UpdateEvaluationContext(true);
 	}
 }
 
@@ -134,7 +134,7 @@ void UGameplayCameraRigComponent::NotifyChangeCameraRigReference()
 {
 	using namespace UE::Cameras;
 
-	if (HasCameraEvaluationContext())
+	if (HasEvaluationContext())
 	{
 		// Sequencer animated some of our parameters... look for those whose value changed, compared to our
 		// cached parameter bag, and re-apply them to the evaluation context.
@@ -173,18 +173,18 @@ void UGameplayCameraRigComponent::PostEditChangeProperty( struct FPropertyChange
 	const FName MemberPropertyName = PropertyChangedEvent.GetMemberPropertyName();
 	if (MemberPropertyName == GET_MEMBER_NAME_CHECKED(UGameplayCameraRigComponent, CameraRigReference))
 	{
-		if (HasCameraEvaluationContext())
+		if (HasEvaluationContext())
 		{
 			if (PropertyChangedEvent.GetPropertyName() == TEXT("CameraRig"))
 			{
 				// The camera rig asset has changed! Recreate the context.
 				GeneratedCameraAsset = nullptr;
-				RecreateEditorWorldCameraEvaluationContext();
+				RecreateEditorWorldEvaluationContext();
 			}
 			else
 			{
 				// Otherwise, maybe one of the parameter overrides has changed. Re-apply them.
-				UpdateCameraEvaluationContext(true);
+				UpdateEvaluationContext(true);
 			}
 		}
 	}
diff --git a/Engine/Plugins/Cameras/GameplayCameras/Source/GameplayCameras/Private/GameFramework/GameplayCamerasPlayerCameraManager.cpp b/Engine/Plugins/Cameras/GameplayCameras/Source/GameplayCameras/Private/GameFramework/GameplayCamerasPlayerCameraManager.cpp
--- a/Engine/Plugins/Cameras/GameplayCameras/Source/GameplayCameras/Private/GameFramework/GameplayCamerasPlayerCameraManager.cpp
+++ b/Engine/Plugins/Cameras/GameplayCameras/Source/GameplayCameras/Private/GameFramework/GameplayCamerasPlayerCameraManager.cpp
@@ -81,119 +81,177 @@ void AGameplayCamerasPlayerCameraManager::ReleasePlayerController()
 	PCOwner = nullptr;
 }
 
-void AGameplayCamerasPlayerCameraManager::ActivatePersistentBaseCameraRig(UCameraRigAsset* CameraRig)
+void AGameplayCamerasPlayerCameraManager::ActivateGameplayCamera(UGameplayCameraComponentBase* GameplayCamera, EGameplayCameraComponentActivationMode ActivationMode)
 {
-	EnsureNullContext();
-	IGameplayCameraSystemHost::ActivateCameraRig(CameraRig, NullContext, ECameraRigLayer::Base);
-}
+	using namespace UE::Cameras;
 
-void AGameplayCamerasPlayerCameraManager::ActivatePersistentGlobalCameraRig(UCameraRigAsset* CameraRig)
-{
-	EnsureNullContext();
-	IGameplayCameraSystemHost::ActivateCameraRig(CameraRig, NullContext, ECameraRigLayer::Global);
-}
+	if (!GameplayCamera)
+	{
+		UE_LOG(LogCameraSystem, Error, TEXT("Can't activate a null Gameplay Camera!"));
+		return;
+	}
 
-void AGameplayCamerasPlayerCameraManager::ActivatePersistentVisualCameraRig(UCameraRigAsset* CameraRig)
-{
-	EnsureNullContext();
-	IGameplayCameraSystemHost::ActivateCameraRig(CameraRig, NullContext, ECameraRigLayer::Visual);
-}
+	if (GameplayCamera->bRunStandaloneCameraSystem)
+	{
+		UE_LOG(LogCameraSystem, Error, 
+				TEXT("Can't activate Gameplay Camera '%s.%s': it is set to run in 'standalone' mode. Please disable 'Run Standalone Camera System' on it."),
+				*GetNameSafe(GameplayCamera->GetOwner()), *GetNameSafe(GameplayCamera));
+		return;
+	}
 
-void AGameplayCamerasPlayerCameraManager::EnsureNullContext()
-{
-	using namespace UE::Cameras;
+	// See if there's a context for this component on the context stack.
+	FCameraEvaluationContextStack& ContextStack = CameraSystemEvaluator->GetEvaluationContextStack();
+	TSharedPtr<FCameraEvaluationContext> GameplayCameraContext = ContextStack.FindContextByOwner(GameplayCamera);
+	const bool bFoundExistingContext = GameplayCameraContext.IsValid();
+	if (!GameplayCameraContext)
+	{
+		// Make sure the component is active, so that it has an evaluation context.
+		APlayerController* PlayerController = GetOwningPlayerController();
+		GameplayCamera->ActivateCameraForPlayerController(PlayerController, false);
+		GameplayCameraContext = GameplayCamera->GetEvaluationContext();
+	}
+	if (!GameplayCameraContext.IsValid())
+	{
+		UE_LOG(LogCameraSystem, Error, 
+				TEXT("Can't activate Gameplay Camera '%s.%s': can't create an evaluation context!"),
+				*GetNameSafe(GameplayCamera->GetOwner()), *GetNameSafe(GameplayCamera));
+		return;
+	}
+	if (GameplayCameraContext->IsActive())
+	{
+		UE_LOG(LogCameraSystem, Error, 
+				TEXT("Can't activate Gameplay Camera '%s.%s': its evaluation context is already active!"),
+				*GetNameSafe(GameplayCamera->GetOwner()), *GetNameSafe(GameplayCamera));
+		return;
+	}
 
-	if (!NullContext)
+	switch (ActivationMode)
 	{
-		FCameraEvaluationContextInitializeParams InitParams;
-		InitParams.Owner = this;
-		InitParams.PlayerController = GetOwningPlayerController();
-		NullContext = MakeShared<FCameraEvaluationContext>(InitParams);
+		case EGameplayCameraComponentActivationMode::Push:
+			{
+				// If the context already existed in the stack, this will move it to the top without re-activating it.
+				// Otherwise, it will activate it and push it on the top.
+				ContextStack.PushContext(GameplayCameraContext.ToSharedRef());
+			}
+			break;
+		case EGameplayCameraComponentActivationMode::InsertOrPush:
+			{
+				// If there's an active context, insert the new one inside it. If our context already existed,
+				// we'll need to deactivate it first.
+				// If there's no active context, simply push or move our context to the top.
+				if (TSharedPtr<FCameraEvaluationContext> ActiveContext = ContextStack.GetActiveContext())
+				{
+					if (bFoundExistingContext)
+					{
+						ContextStack.RemoveContext(GameplayCameraContext.ToSharedRef());
+					}
+
+					const bool bSuccess = ActiveContext->AddChildContext(GameplayCameraContext.ToSharedRef());
+					if (!bSuccess)
+					{
+						UE_LOG(LogCameraSystem, Error,
+								TEXT("Couldn't insert camera director for '{0}.{1}' inside camera director for '{2}.{3}'. "
+									 "Activation failed."),
+								*GetNameSafe(GameplayCamera->GetOwner()), *GetNameSafe(GameplayCamera));
+					}
+				}
+				else
+				{
+					ContextStack.PushContext(GameplayCameraContext.ToSharedRef());
+				}
+			}
+			break;
+		default:
+			ensureMsgf(false, TEXT("Unknown activation mode: %s"), *UEnum::GetValueAsString(ActivationMode));
+			break;
 	}
 }
 
-FCameraRigInstanceID AGameplayCamerasPlayerCameraManager::StartGlobalCameraModifierRig(const UCameraRigAsset* CameraRig, int32 OrderKey)
+void AGameplayCamerasPlayerCameraManager::DeactivateGameplayCamera(UGameplayCameraComponentBase* GameplayCamera, bool bDeactivateAllCameraRigs)
 {
 	using namespace UE::Cameras;
 
-	if (CameraSystemEvaluator)
+	FCameraEvaluationContextStack& ContextStack = CameraSystemEvaluator->GetEvaluationContextStack();
+	TSharedPtr<FCameraEvaluationContext> ExistingContext = ContextStack.FindContextByOwner(GameplayCamera);
+	if (ExistingContext)
 	{
-		EnsureNullContext();
+		if (bDeactivateAllCameraRigs)
+		{
+			FRootCameraNodeEvaluator* RootEvaluator = CameraSystemEvaluator->GetRootNodeEvaluator();
+			RootEvaluator->DeactivateAllCameraRigs(ExistingContext, true);
+		}
 
-		TSharedPtr<FCameraModifierService> CameraModifierService = CameraSystemEvaluator->FindEvaluationService<FCameraModifierService>();
-		return CameraModifierService->StartCameraModifierRig(CameraRig, NullContext.ToSharedRef(), ECameraRigLayer::Global, OrderKey);
+		ContextStack.RemoveContext(ExistingContext.ToSharedRef());
 	}
+}
+
+FCameraRigInstanceID AGameplayCamerasPlayerCameraManager::ActivatePersistentBaseCameraRig(UCameraRigAsset* CameraRig)
+{
+	EnsureNullContext();
+	return IGameplayCameraSystemHost::ActivateCameraRig(CameraRig, NullContext, ECameraRigLayer::Base);
+}
 
-	return FCameraRigInstanceID();
+FCameraRigInstanceID AGameplayCamerasPlayerCameraManager::ActivatePersistentGlobalCameraRig(UCameraRigAsset* CameraRig)
+{
+	EnsureNullContext();
+	return IGameplayCameraSystemHost::ActivateCameraRig(CameraRig, NullContext, ECameraRigLayer::Global);
 }
 
-FCameraRigInstanceID AGameplayCamerasPlayerCameraManager::StartVisualCameraModifierRig(const UCameraRigAsset* CameraRig, int32 OrderKey)
+FCameraRigInstanceID AGameplayCamerasPlayerCameraManager::ActivatePersistentVisualCameraRig(UCameraRigAsset* CameraRig)
 {
-	using namespace UE::Cameras;
+	EnsureNullContext();
+	return IGameplayCameraSystemHost::ActivateCameraRig(CameraRig, NullContext, ECameraRigLayer::Visual);
+}
 
-	if (CameraSystemEvaluator)
-	{
-		EnsureNullContext();
+void AGameplayCamerasPlayerCameraManager::DeactivateCameraRig(FCameraRigInstanceID InstanceID, bool bImmediately)
+{
+	EnsureNullContext();
+	IGameplayCameraSystemHost::DeactivateCameraRig(InstanceID, bImmediately);
+}
 
-		TSharedPtr<FCameraModifierService> CameraModifierService = CameraSystemEvaluator->FindEvaluationService<FCameraModifierService>();
-		return CameraModifierService->StartCameraModifierRig(CameraRig, NullContext.ToSharedRef(), ECameraRigLayer::Visual, OrderKey);
-	}
+FCameraRigInstanceID AGameplayCamerasPlayerCameraManager::StartGlobalCameraModifierRig(const UCameraRigAsset* CameraRig, int32 OrderKey)
+{
+	EnsureNullContext();
+	return IGameplayCameraSystemHost::StartCameraModifierRig(CameraRig, NullContext, ECameraRigLayer::Global, OrderKey);
+}
 
-	return FCameraRigInstanceID();
+FCameraRigInstanceID AGameplayCamerasPlayerCameraManager::StartVisualCameraModifierRig(const UCameraRigAsset* CameraRig, int32 OrderKey)
+{
+	EnsureNullContext();
+	return IGameplayCameraSystemHost::StartCameraModifierRig(CameraRig, NullContext, ECameraRigLayer::Visual, OrderKey);
 }
 
 void AGameplayCamerasPlayerCameraManager::StopCameraModifierRig(FCameraRigInstanceID InstanceID, bool bImmediately)
 {
-	using namespace UE::Cameras;
-
-	if (CameraSystemEvaluator)
-	{
-		TSharedPtr<FCameraModifierService> CameraModifierService = CameraSystemEvaluator->FindEvaluationService<FCameraModifierService>();
-		return CameraModifierService->StopCameraModifierRig(InstanceID, bImmediately);
-	}
+	IGameplayCameraSystemHost::StopCameraModifierRig(InstanceID, bImmediately);
 }
 
 FCameraShakeInstanceID AGameplayCamerasPlayerCameraManager::StartCameraShakeAsset(const UCameraShakeAsset* CameraShake, float ShakeScale, ECameraShakePlaySpace PlaySpace, FRotator UserPlaySpaceRotation)
 {
-	using namespace UE::Cameras;
-
-	if (CameraSystemEvaluator)
-	{
-		FStartCameraShakeParams Params;
-		Params.CameraShake = CameraShake;
-		Params.ShakeScale = ShakeScale;
-		Params.PlaySpace = PlaySpace;
-		Params.UserPlaySpaceRotation = UserPlaySpaceRotation;
-
-		TSharedPtr<FCameraShakeService> CameraShakeService = CameraSystemEvaluator->FindEvaluationService<FCameraShakeService>();
-		return CameraShakeService->StartCameraShake(Params);
-	}
-
-	return FCameraShakeInstanceID();
+	return IGameplayCameraSystemHost::StartCameraShake(CameraShake, ShakeScale, PlaySpace, UserPlaySpaceRotation);
 }
 
 bool AGameplayCamerasPlayerCameraManager::IsCameraShakeAssetPlaying(FCameraShakeInstanceID InInstanceID) const
 {
-	using namespace UE::Cameras;
-
-	if (CameraSystemEvaluator)
-	{
-		TSharedPtr<FCameraShakeService> CameraShakeService = CameraSystemEvaluator->FindEvaluationService<FCameraShakeService>();
-		return CameraShakeService->IsCameraShakePlaying(InInstanceID);
-	}
-	return false;
+	return IGameplayCameraSystemHost::IsCameraShakePlaying(InInstanceID);
 }
 
 bool AGameplayCamerasPlayerCameraManager::StopCameraShakeAsset(FCameraShakeInstanceID InInstanceID, bool bImmediately)
+{
+	return IGameplayCameraSystemHost::StopCameraShake(InInstanceID, bImmediately);
+}
+
+void AGameplayCamerasPlayerCameraManager::EnsureNullContext()
 {
 	using namespace UE::Cameras;
 
-	if (CameraSystemEvaluator)
+	if (!NullContext)
 	{
-		TSharedPtr<FCameraShakeService> CameraShakeService = CameraSystemEvaluator->FindEvaluationService<FCameraShakeService>();
-		return CameraShakeService->StopCameraShake(InInstanceID, bImmediately);
+		FCameraEvaluationContextInitializeParams InitParams;
+		InitParams.Owner = this;
+		InitParams.PlayerController = GetOwningPlayerController();
+		NullContext = MakeShared<FCameraEvaluationContext>(InitParams);
 	}
-	return false;
 }
 
 void AGameplayCamerasPlayerCameraManager::InitializeFor(APlayerController* PlayerController)
@@ -257,30 +315,21 @@ void AGameplayCamerasPlayerCameraManager::SetViewTarget(AActor* NewViewTarget, F
 
 	// See if we can find the view target in the context stack. If so, reactivate it instead of potentially
 	// making a new context for the same thing.
-	bool bFoundContext = false;
-	TArray<TSharedPtr<FCameraEvaluationContext>> CurrentContexts;
-	ContextStack.GetAllContexts(CurrentContexts);
-	for (TSharedPtr<FCameraEvaluationContext> CurrentContext : CurrentContexts)
+	TSharedPtr<FCameraEvaluationContext> ExistingContext = ContextStack.FindContextByPredicate(
+			[NewViewTarget](TSharedRef<FCameraEvaluationContext> Item)
+			{
+				UObject* Owner = Item->GetOwner();
+				return Owner && (Owner == NewViewTarget || Owner->GetTypedOuter<AActor>() == NewViewTarget);
+			});
+	if (ExistingContext)
 	{
-		UObject* CurrentContextOwner = CurrentContext ? CurrentContext->GetOwner() : nullptr;
-		if (CurrentContextOwner && 
-				(CurrentContext->GetOwner() == NewViewTarget || 
-				 CurrentContextOwner->GetTypedOuter<AActor>() == NewViewTarget))
-		{
-			// This will move the context to the top if it's already in the stack (which it is, we
-			// found it there).
-			ContextStack.PushContext(CurrentContext.ToSharedRef());
-			bFoundContext = true;
-		}
+		// This will move the context to the top if it's already in the stack (which it is, we found it there).
+		ContextStack.PushContext(ExistingContext.ToSharedRef());
 	}
-
-	if (!bFoundContext)
+	else
 	{
-		if (UGameplayCameraComponentBase* GameplayCameraComponent = NewViewTarget->FindComponentByClass<UGameplayCameraComponentBase>())
-		{
-			GameplayCameraComponent->ActivateCameraForPlayerController(PCOwner);
-		}
-		else if (UCameraComponent* CameraComponent = NewViewTarget->FindComponentByClass<UCameraComponent>())
+		// Create a new context for this view target.
+		if (UCameraComponent* CameraComponent = NewViewTarget->FindComponentByClass<UCameraComponent>())
 		{
 			TSharedRef<FActorCameraEvaluationContext> NewContext = MakeShared<FActorCameraEvaluationContext>(CameraComponent);
 			CameraSystemEvaluator->PushEvaluationContext(NewContext);
diff --git a/Engine/Plugins/Cameras/GameplayCameras/Source/GameplayCameras/Private/GameFramework/IGameplayCameraSystemHost.cpp b/Engine/Plugins/Cameras/GameplayCameras/Source/GameplayCameras/Private/GameFramework/IGameplayCameraSystemHost.cpp
--- a/Engine/Plugins/Cameras/GameplayCameras/Source/GameplayCameras/Private/GameFramework/IGameplayCameraSystemHost.cpp
+++ b/Engine/Plugins/Cameras/GameplayCameras/Source/GameplayCameras/Private/GameFramework/IGameplayCameraSystemHost.cpp
@@ -15,6 +15,8 @@
 #include "Engine/Canvas.h"
 #include "Engine/World.h"
 #include "GameFramework/PlayerController.h"
+#include "Services/CameraModifierService.h"
+#include "Services/CameraShakeService.h"
 
 #include UE_INLINE_GENERATED_CPP_BY_NAME(IGameplayCameraSystemHost)
 
@@ -60,7 +62,10 @@ TScriptInterface<IGameplayCameraSystemHost> IGameplayCameraSystemHost::GetAsScri
 	return Result;
 }
 
-void IGameplayCameraSystemHost::ActivateCameraRig(UCameraRigAsset* CameraRig, TSharedPtr<UE::Cameras::FCameraEvaluationContext> EvaluationContext, ECameraRigLayer EvaluationLayer)
+FCameraRigInstanceID IGameplayCameraSystemHost::ActivateCameraRig(
+		const UCameraRigAsset* CameraRig, 
+		TSharedPtr<FCameraEvaluationContext> EvaluationContext,
+		ECameraRigLayer EvaluationLayer)
 {
 	using namespace UE::Cameras;
 
@@ -69,7 +74,7 @@ void IGameplayCameraSystemHost::ActivateCameraRig(UCameraRigAsset* CameraRig, TS
 		FFrame::KismetExecutionMessage(
 				*FString::Printf(TEXT("Can't activate camera rig on '%s': no camera rig given!"), *GetNameSafe(GetAsObject())),
 				ELogVerbosity::Error);
-		return;
+		return FCameraRigInstanceID();
 	}
 	
 	if (!EvaluationContext)
@@ -79,11 +84,9 @@ void IGameplayCameraSystemHost::ActivateCameraRig(UCameraRigAsset* CameraRig, TS
 					TEXT("Can't activate camera rig '%s' on '%s': invalid evaluation context given!"), 
 					*GetNameSafe(CameraRig), *GetNameSafe(GetAsObject())),
 				ELogVerbosity::Error);
-		return;
+		return FCameraRigInstanceID();
 	}
 
-	EnsureCameraSystemInitialized();
-
 	if (ensure(CameraSystemEvaluator))
 	{
 		FActivateCameraRigParams Params;
@@ -92,8 +95,99 @@ void IGameplayCameraSystemHost::ActivateCameraRig(UCameraRigAsset* CameraRig, TS
 		Params.Layer = EvaluationLayer;
 
 		FRootCameraNodeEvaluator* RootNodeEvaluator = CameraSystemEvaluator->GetRootNodeEvaluator();
-		RootNodeEvaluator->ActivateCameraRig(Params);
+		return RootNodeEvaluator->ActivateCameraRig(Params);
+	}
+
+	return FCameraRigInstanceID();
+}
+
+void IGameplayCameraSystemHost::DeactivateCameraRig(FCameraRigInstanceID InInstanceID, bool bImmediately)
+{
+	using namespace UE::Cameras;
+
+	if (InInstanceID.IsValid())
+	{
+		FDeactivateCameraRigParams Params;
+		Params.InstanceID = InInstanceID;
+		Params.bDeactiveImmediately = bImmediately;
+
+		FRootCameraNodeEvaluator* RootNodeEvaluator = CameraSystemEvaluator->GetRootNodeEvaluator();
+		RootNodeEvaluator->DeactivateCameraRig(Params);
+	}
+}
+
+FCameraRigInstanceID IGameplayCameraSystemHost::StartCameraModifierRig(
+		const UCameraRigAsset* CameraRig,
+		TSharedPtr<FCameraEvaluationContext> EvaluationContext,
+		ECameraRigLayer EvaluationLayer,
+		int32 OrderKey)
+{
+	using namespace UE::Cameras;
+
+	if (ensure(CameraSystemEvaluator))
+	{
+		TSharedPtr<FCameraModifierService> CameraModifierService = CameraSystemEvaluator->FindEvaluationService<FCameraModifierService>();
+		if (ensure(CameraModifierService))
+		{
+			return CameraModifierService->StartCameraModifierRig(CameraRig, EvaluationContext.ToSharedRef(), ECameraRigLayer::Global, OrderKey);
+		}
+	}
+
+	return FCameraRigInstanceID();
+}
+
+void IGameplayCameraSystemHost::StopCameraModifierRig(FCameraRigInstanceID InstanceID, bool bImmediately)
+{
+	using namespace UE::Cameras;
+
+	if (ensure(CameraSystemEvaluator))
+	{
+		TSharedPtr<FCameraModifierService> CameraModifierService = CameraSystemEvaluator->FindEvaluationService<FCameraModifierService>();
+		CameraModifierService->StopCameraModifierRig(InstanceID, bImmediately);
+	}
+}
+
+FCameraShakeInstanceID IGameplayCameraSystemHost::StartCameraShake(const UCameraShakeAsset* CameraShake, float ShakeScale, ECameraShakePlaySpace PlaySpace, FRotator UserPlaySpaceRotation)
+{
+	using namespace UE::Cameras;
+
+	if (ensure(CameraSystemEvaluator))
+	{
+		FStartCameraShakeParams Params;
+		Params.CameraShake = CameraShake;
+		Params.ShakeScale = ShakeScale;
+		Params.PlaySpace = PlaySpace;
+		Params.UserPlaySpaceRotation = UserPlaySpaceRotation;
+
+		TSharedPtr<FCameraShakeService> CameraShakeService = CameraSystemEvaluator->FindEvaluationService<FCameraShakeService>();
+		return CameraShakeService->StartCameraShake(Params);
+	}
+
+	return FCameraShakeInstanceID();
+}
+
+bool IGameplayCameraSystemHost::IsCameraShakePlaying(FCameraShakeInstanceID InInstanceID) const
+{
+	using namespace UE::Cameras;
+
+	if (ensure(CameraSystemEvaluator))
+	{
+		TSharedPtr<FCameraShakeService> CameraShakeService = CameraSystemEvaluator->FindEvaluationService<FCameraShakeService>();
+		return CameraShakeService->IsCameraShakePlaying(InInstanceID);
+	}
+	return false;
+}
+
+bool IGameplayCameraSystemHost::StopCameraShake(FCameraShakeInstanceID InInstanceID, bool bImmediately)
+{
+	using namespace UE::Cameras;
+
+	if (ensure(CameraSystemEvaluator))
+	{
+		TSharedPtr<FCameraShakeService> CameraShakeService = CameraSystemEvaluator->FindEvaluationService<FCameraShakeService>();
+		return CameraShakeService->StopCameraShake(InInstanceID, bImmediately);
 	}
+	return false;
 }
 
 IGameplayCameraSystemHost* IGameplayCameraSystemHost::FindActiveHost(APlayerController* PlayerController)
diff --git a/Engine/Plugins/Cameras/GameplayCameras/Source/GameplayCameras/Public/Core/CameraEvaluationContextStack.h b/Engine/Plugins/Cameras/GameplayCameras/Source/GameplayCameras/Public/Core/CameraEvaluationContextStack.h
--- a/Engine/Plugins/Cameras/GameplayCameras/Source/GameplayCameras/Public/Core/CameraEvaluationContextStack.h
+++ b/Engine/Plugins/Cameras/GameplayCameras/Source/GameplayCameras/Public/Core/CameraEvaluationContextStack.h
@@ -4,6 +4,7 @@
 
 #include "Containers/Array.h"
 #include "GameplayCameras.h"
+#include "Templates/Invoke.h"
 #include "UObject/WeakObjectPtr.h"
 
 class UCameraAsset;
@@ -61,6 +62,17 @@ public:
 	/** Gets all the contexts in the stack, from bottom to top. */
 	void GetAllContexts(TArray<TSharedPtr<FCameraEvaluationContext>>& OutContexts) const;
 
+	/** Finds a context owned by the given object, if any. */
+	TSharedPtr<FCameraEvaluationContext> FindContextByOwner(const UObject* InOwner) const;
+
+	/**
+	 * Finds a context that fulfills the given predicate.
+	 *
+	 * The predicate function must have the signature: (TSharedRef<FCameraEvaluationContext>) -> bool
+	 */
+	template<typename Predicate>
+	TSharedPtr<FCameraEvaluationContext> FindContextByPredicate(Predicate Pred) const;
+
 	/** Empties the stack of all contexts. */
 	void Reset();
 
@@ -95,5 +107,21 @@ private:
 #endif  // UE_GAMEPLAY_CAMERAS_DEBUG
 };
 
+template<typename Predicate>
+TSharedPtr<FCameraEvaluationContext> FCameraEvaluationContextStack::FindContextByPredicate(Predicate Pred) const
+{
+	for (const FContextEntry& Entry : Entries)
+	{
+		if (TSharedPtr<FCameraEvaluationContext> Context = Entry.WeakContext.Pin())
+		{
+			if (::Invoke(Pred, Context.ToSharedRef()))
+			{
+				return Context;
+			}
+		}
+	}
+	return TSharedPtr<FCameraEvaluationContext>();
+}
+
 }  // namespace UE::Cameras
 
diff --git a/Engine/Plugins/Cameras/GameplayCameras/Source/GameplayCameras/Public/GameFramework/GameplayCameraComponentBase.h b/Engine/Plugins/Cameras/GameplayCameras/Source/GameplayCameras/Public/GameFramework/GameplayCameraComponentBase.h
--- a/Engine/Plugins/Cameras/GameplayCameras/Source/GameplayCameras/Public/GameFramework/GameplayCameraComponentBase.h
+++ b/Engine/Plugins/Cameras/GameplayCameras/Source/GameplayCameras/Public/GameFramework/GameplayCameraComponentBase.h
@@ -31,20 +31,6 @@ class FGameplayCameraComponentEvaluationContext;
 
 }  // namespace UE::Cameras
 
-/**
- * Defines how to activate a gameplay camera component.
- */
-UENUM()
-enum class EGameplayCameraComponentActivationMode : uint8
-{
-	/** Push the camera director over any existing ones. */
-	Push,
-	/** Push the camera director and try to insert any active one as a child. */
-	PushAndInsert,
-	/** Inserts the camera director as a child of the active one, or push it if there is no active one. */
-	InsertOrPush
-};
-
 /**
  * A component that can run a camera asset inside its own camera evaluation context.
  */
@@ -63,12 +49,17 @@ public:
 	/** Create a new camera component. */
 	GAMEPLAYCAMERAS_API UGameplayCameraComponentBase(const FObjectInitializer& ObjectInit);
 
+public:
+
 	/** Get the camera evaluation context used by this component. */
 	GAMEPLAYCAMERAS_API TSharedPtr<const UE::Cameras::FCameraEvaluationContext> GetEvaluationContext() const;
 
 	/** Get the camera evaluation context used by this component. */
 	GAMEPLAYCAMERAS_API TSharedPtr<UE::Cameras::FCameraEvaluationContext> GetEvaluationContext();
 
+	/** Returns whether an evaluation context exists in this component. */
+	bool HasEvaluationContext() const { return EvaluationContext.IsValid(); }
+
 public:
 
 	/** Gets the child camera component used as the "output" for the gameplay/procedural camera. */
@@ -80,42 +71,18 @@ public:
 	 *
 	 * @param PlayerIndex        The player to activate the camera for.
 	 * @param bSetAsViewTarget   Whether to set this component's actor as the view target for the player.
-	 * @param ActivationMode     How to activate this camera into the player's camera system. Only valid 
-	 *                           and used when the player camera manager is running the camera system.
-	 *                           Must be 'Push' otherwise, when this component runs as a standalone camera
-	 *                           system.
 	 */
 	UFUNCTION(BlueprintCallable, Category=Camera)
-	GAMEPLAYCAMERAS_API void ActivateCameraForPlayerIndex(
-			int32 PlayerIndex, 
-			bool bSetAsViewTarget = true,
-			EGameplayCameraComponentActivationMode ActivationMode = EGameplayCameraComponentActivationMode::Push);
+	GAMEPLAYCAMERAS_API void ActivateCameraForPlayerIndex(int32 PlayerIndex, bool bSetAsViewTarget = true);
 
 	/** 
 	 * Activates the camera for the given player.
 	 *
 	 * @param PlayerController   The player to activate the camera for.
 	 * @param bSetAsViewTarget   Whether to set this component's actor as the view target for the player.
-	 * @param ActivationMode     How to activate this camera into the player's camera system. Only valid 
-	 *                           and used when the player camera manager is running the camera system.
-	 *                           Must be 'Push' otherwise, when this component runs as a standalone camera
-	 *                           system.
-	 */
-	UFUNCTION(BlueprintCallable, Category=Camera)
-	GAMEPLAYCAMERAS_API void ActivateCameraForPlayerController(
-			APlayerController* PlayerController,
-			bool bSetAsViewTarget = true,
-			EGameplayCameraComponentActivationMode ActivationMode = EGameplayCameraComponentActivationMode::Push);
-
-	/** 
-	 * Deactivates the camera.
-	 *
-	 * @param bImmediately       Whether to let this component's camera rigs gracefully blend out before
-	 *                           deactivating. If true, any running camera rigs will be frozen or forcibly
-	 *                           removed from the camera system.
 	 */
 	UFUNCTION(BlueprintCallable, Category=Camera)
-	GAMEPLAYCAMERAS_API void DeactivateCamera(bool bImmediately = false);
+	GAMEPLAYCAMERAS_API void ActivateCameraForPlayerController(APlayerController* PlayerController, bool bSetAsViewTarget = true);
 
 	/** Gets the shared camera evaluation data for this component's evaluation context. */
 	UFUNCTION(BlueprintPure, Category=Camera, meta=(DisplayName="Get Shared Camera Data"))
@@ -133,15 +100,19 @@ public:
 
 	/** Activates the given camera rig prefab in the base layer. */
 	UFUNCTION(BlueprintCallable, Category="Camera")
-	void ActivatePersistentBaseCameraRig(UCameraRigAsset* CameraRig);
+	FCameraRigInstanceID ActivatePersistentBaseCameraRig(UCameraRigAsset* CameraRig);
 
 	/** Activates the given camera rig prefab in the global layer. */
 	UFUNCTION(BlueprintCallable, Category="Camera")
-	void ActivatePersistentGlobalCameraRig(UCameraRigAsset* CameraRig);
+	FCameraRigInstanceID ActivatePersistentGlobalCameraRig(UCameraRigAsset* CameraRig);
 
 	/** Activates the given camera rig prefab in the visual layer. */
 	UFUNCTION(BlueprintCallable, Category="Camera")
-	void ActivatePersistentVisualCameraRig(UCameraRigAsset* CameraRig);
+	FCameraRigInstanceID ActivatePersistentVisualCameraRig(UCameraRigAsset* CameraRig);
+
+	/** Deactivates a previously activated base, global, or visual camera rig. */
+	UFUNCTION(BlueprintCallable, Category="Camera")
+	void DeactivateCameraRig(FCameraRigInstanceID InstanceID, bool bImmediately = false);
 
 public:
 
@@ -204,43 +175,47 @@ public:
 
 public:
 
-	bool CalcCamera(float DeltaTime, FMinimalViewInfo& OutResult);
+	// Internal API
 
 #if WITH_EDITOR
 	GAMEPLAYCAMERAS_API void OnDrawVisualizationHUD(const FViewport* Viewport, const FSceneView* SceneView, FCanvas* Canvas) const;
 #endif  // WITH_EDITOR
 
+public:
+
+	// Deprecated API
+	
+	UFUNCTION(BlueprintCallable, Category="Camera", meta=(DeprecatedFunction, DeprecatedMessage="Please use Deactivate"))
+	void DeactivateCamera(bool bImmediately = false) { Deactivate(); }
+
 protected:
 
 	// UGameplayCameraComponentBase interface.
 	virtual UCameraAsset* GetCameraAsset() PURE_VIRTUAL(UGameplayCameraComponentBase::GetCameraAsset, return nullptr;)
 	virtual void OnUpdateCameraEvaluationContext(bool bForceApplyParameterOverrides) {}
 
-	void UpdateCameraEvaluationContext(bool bForceApplyParameterOverrides);
-	bool HasCameraEvaluationContext() const { return EvaluationContext.IsValid(); }
+	void UpdateEvaluationContext(bool bForceApplyParameterOverrides);
 
 	bool IsEditorWorld() const;
 	void UpdateControlRotationIfNeeded();
 
 #if WITH_EDITOR
-	void ReinitializeCameraEvaluationContext(
+	void ReinitializeEvaluationContext(
 			const FCameraVariableTableAllocationInfo& VariableTableAllocationInfo,
 			const FCameraContextDataTableAllocationInfo& ContextDataTableAllocationInfo);
-	void RecreateEditorWorldCameraEvaluationContext();
+	void RecreateEditorWorldEvaluationContext();
 #endif  // WITH_EDITOR
 
 private:
 
 	bool CanRunCameraSystem() const;
-	void EnsureCameraSystemHost();
+	bool EnsureCameraSystemHostIfNeeded();
+	void DestroyCameraSystemHost();
 
-	void CreateCameraEvaluationContext(APlayerController* PlayerController);
-	void DestroyCameraEvaluationContext();
+	void EnsureEvaluationContext(APlayerController* PlayerController);
+	void DestroyEvaluationContext();
 
-	void ActivateCameraEvaluationContext(APlayerController* PlayerController, IGameplayCameraSystemHost* Host, EGameplayCameraComponentActivationMode ActivationMode);
 	void UpdateOutputCameraComponent();
-	void DeactivateCameraEvaluationContext(bool bImmediately);
-	void CheckPendingDeactivation();
 
 #if WITH_EDITOR
 	void AutoManageEditorPreviewEvaluator();
@@ -256,13 +231,22 @@ public:
 	UPROPERTY(BlueprintReadWrite, EditAnywhere, Category=Activation, meta=(EditCondition="bAutoActivate"))
 	TEnumAsByte<EAutoReceiveInput::Type> AutoActivateForPlayer;
 
+	/**
+	 * When enabled, this camera component runs its own camera system when activated, which
+	 * makes it posisble for it to run independently. When disabled, this camera component
+	 * only creates an evaluation context that must be run with an external camera system,
+	 * such as the Gameplay Cameras Player Camera Manager.
+	 */
+	UPROPERTY(BlueprintReadWrite, EditAnywhere, Category=Camera)
+	bool bRunStandaloneCameraSystem = true;
+
 	/**
 	 * Specifies whether this component should set the player controller's control rotation 
 	 * to the computed point of view's orientation every frame. This is only used when a 
 	 * player controller is associated with this component, and the view target is that
 	 * component.
 	 */
-	UPROPERTY(BlueprintReadWrite, EditAnywhere, Category=Camera)
+	UPROPERTY(BlueprintReadWrite, EditAnywhere, Category=Camera, meta=(EditCondition="bRunStandaloneCameraSystem"))
 	bool bSetControlRotationWhenViewTarget = false;
 
 #if WITH_EDITORONLY_DATA
@@ -290,15 +274,9 @@ private:
 	/** Evaluation context. */
 	TSharedPtr<FGameplayCameraComponentEvaluationContext> EvaluationContext;
 
-	/** The camera system in which to check for running rigs before deactivating our evaluation context. */
-	TWeakPtr<FCameraSystemEvaluator> PendingDeactivateCameraSystemEvaluator;
-
 	/** Whether to force a camera cut next frame. */
 	bool bIsCameraCutNextFrame = false;
 
-	/** Whether to check for a pending deactivation. */
-	bool bIsPendingDeactivate = false;
-
 #if WITH_EDITOR
 	
 	/** Whether this component is running in an editor world. */
diff --git a/Engine/Plugins/Cameras/GameplayCameras/Source/GameplayCameras/Public/GameFramework/GameplayCamerasPlayerCameraManager.h b/Engine/Plugins/Cameras/GameplayCameras/Source/GameplayCameras/Public/GameFramework/GameplayCamerasPlayerCameraManager.h
--- a/Engine/Plugins/Cameras/GameplayCameras/Source/GameplayCameras/Public/GameFramework/GameplayCamerasPlayerCameraManager.h
+++ b/Engine/Plugins/Cameras/GameplayCameras/Source/GameplayCameras/Public/GameFramework/GameplayCamerasPlayerCameraManager.h
@@ -11,6 +11,7 @@
 
 #define UE_API GAMEPLAYCAMERAS_API
 
+class UGameplayCameraComponentBase;
 struct FCameraRigInstanceID;
 
 namespace UE::Cameras
@@ -22,6 +23,18 @@ class FViewTargetContextReferencerService;
 
 }  // namespace UE::Cameras
 
+/**
+ * Defines how to activate a gameplay camera component.
+ */
+UENUM()
+enum class EGameplayCameraComponentActivationMode : uint8
+{
+	/** Push the camera director over any existing ones. */
+	Push,
+	/** Inserts the camera director as a child of the active one, or push it if there is no active one. */
+	InsertOrPush
+};
+
 /**
  * Defines how the GameplayCamerasPlayerCameraManager should handle the player's view rotation.
  */
@@ -81,19 +94,33 @@ public:
 	UFUNCTION(BlueprintCallable, Category="Camera")
 	void ReleasePlayerController();
 
+public:
+
+	/** Activates the given component inside the manager's camera system. */
+	UFUNCTION(BlueprintCallable, Category="Camera")
+	void ActivateGameplayCamera(UGameplayCameraComponentBase* GameplayCamera, EGameplayCameraComponentActivationMode ActivationMode = EGameplayCameraComponentActivationMode::Push);
+
+	/** Deactivates a previously activated component. */
+	UFUNCTION(BlueprintCallable, Category="Camera")
+	void DeactivateGameplayCamera(UGameplayCameraComponentBase* GameplayCamera, bool bDeactivateAllCameraRigs = false);
+
 public:
 
 	/** Activates the given camera rig prefab in the base layer. */
 	UFUNCTION(BlueprintCallable, Category="Camera")
-	void ActivatePersistentBaseCameraRig(UCameraRigAsset* CameraRig);
+	FCameraRigInstanceID ActivatePersistentBaseCameraRig(UCameraRigAsset* CameraRig);
 
 	/** Activates the given camera rig prefab in the global layer. */
 	UFUNCTION(BlueprintCallable, Category="Camera")
-	void ActivatePersistentGlobalCameraRig(UCameraRigAsset* CameraRig);
+	FCameraRigInstanceID ActivatePersistentGlobalCameraRig(UCameraRigAsset* CameraRig);
 
 	/** Activates the given camera rig prefab in the visual layer. */
 	UFUNCTION(BlueprintCallable, Category="Camera")
-	void ActivatePersistentVisualCameraRig(UCameraRigAsset* CameraRig);
+	FCameraRigInstanceID ActivatePersistentVisualCameraRig(UCameraRigAsset* CameraRig);
+
+	/** Deactivates a previously activated base, global, or visual camera rig. */
+	UFUNCTION(BlueprintCallable, Category="Camera")
+	void DeactivateCameraRig(FCameraRigInstanceID InstanceID, bool bImmediately = false);
 
 public:
 
diff --git a/Engine/Plugins/Cameras/GameplayCameras/Source/GameplayCameras/Public/GameFramework/IGameplayCameraSystemHost.h b/Engine/Plugins/Cameras/GameplayCameras/Source/GameplayCameras/Public/GameFramework/IGameplayCameraSystemHost.h
--- a/Engine/Plugins/Cameras/GameplayCameras/Source/GameplayCameras/Public/GameFramework/IGameplayCameraSystemHost.h
+++ b/Engine/Plugins/Cameras/GameplayCameras/Source/GameplayCameras/Public/GameFramework/IGameplayCameraSystemHost.h
@@ -2,6 +2,9 @@
 
 #pragma once
 
+#include "Camera/CameraTypes.h"
+#include "Core/CameraRigInstanceID.h"
+#include "Core/CameraShakeInstanceID.h"
 #include "GameplayCameras.h"
 #include "Templates/SharedPointerFwd.h"
 #include "UObject/Interface.h"
@@ -9,6 +12,7 @@
 #include "IGameplayCameraSystemHost.generated.h"
 
 class UCameraRigAsset;
+class UCameraShakeAsset;
 class UCanvas;
 enum class ECameraRigLayer : uint8;
 
@@ -38,6 +42,7 @@ class IGameplayCameraSystemHost
 public:
 
 	using FCameraSystemEvaluator = UE::Cameras::FCameraSystemEvaluator;
+	using FCameraEvaluationContext = UE::Cameras::FCameraEvaluationContext;
 
 	/** Gets the camera system evaluator. */
 	GAMEPLAYCAMERAS_API TSharedPtr<FCameraSystemEvaluator> GetCameraSystemEvaluator();
@@ -73,14 +78,45 @@ protected:
 	/** Updates the camera system, if it exists. */
 	GAMEPLAYCAMERAS_API void UpdateCameraSystem(float DeltaTime);
 
-	/** Activates the given camera rig in the given layer. Should not be used with Main layer. */
-	GAMEPLAYCAMERAS_API void ActivateCameraRig(UCameraRigAsset* CameraRig, TSharedPtr<UE::Cameras::FCameraEvaluationContext> EvaluationContext, ECameraRigLayer EvaluationLayer);
-
 #if WITH_EDITOR
 	/** Updates the camera system, if it exists, for an editor world preview. */
 	GAMEPLAYCAMERAS_API void UpdateCameraSystemForEditorPreview(float DeltaTime);
 #endif  // WITH_EDITOR
 
+protected:
+
+	/** Activates the given camera rig in the given layer. Should not be used with Main layer. */
+	GAMEPLAYCAMERAS_API FCameraRigInstanceID ActivateCameraRig(
+			const UCameraRigAsset* CameraRig,
+			TSharedPtr<FCameraEvaluationContext> EvaluationContext,
+			ECameraRigLayer EvaluationLayer);
+
+	/** Deactivates a previously activated camera rig. */
+	GAMEPLAYCAMERAS_API void DeactivateCameraRig(FCameraRigInstanceID InInstanceID, bool bImmediately = false);
+
+	/** Starts a camera modifier rig on the given layer. */
+	GAMEPLAYCAMERAS_API FCameraRigInstanceID StartCameraModifierRig(
+			const UCameraRigAsset* CameraRig,
+			TSharedPtr<FCameraEvaluationContext> EvaluationContext,
+			ECameraRigLayer EvaluationLayer,
+			int32 OrderKey = 0);
+
+	/** Stops a camera modifier rig on previously started on the global or visual layer. */
+	GAMEPLAYCAMERAS_API void StopCameraModifierRig(FCameraRigInstanceID InstanceID, bool bImmediately = false);
+
+	/** Starts a new camera shake. */
+	GAMEPLAYCAMERAS_API FCameraShakeInstanceID StartCameraShake(
+			const UCameraShakeAsset* CameraShake,
+			float ShakeScale = 1.f,
+			ECameraShakePlaySpace PlaySpace = ECameraShakePlaySpace::CameraLocal,
+			FRotator UserPlaySpaceRotation = FRotator::ZeroRotator);
+
+	/** Checks if a camera shake is running. */
+	GAMEPLAYCAMERAS_API bool IsCameraShakePlaying(FCameraShakeInstanceID InInstanceID) const;
+
+	/** Stops a running camera shake. */
+	GAMEPLAYCAMERAS_API bool StopCameraShake(FCameraShakeInstanceID InInstanceID, bool bImmediately = false);
+
 protected:
 
 	/** The camera system evaluator. */
