Cameras: make the boom arm and collision push node evaluators extensible

diff --git a/Engine/Plugins/Cameras/GameplayCameras/Source/GameplayCameras/Private/Nodes/Collision/CollisionPushCameraNode.cpp b/Engine/Plugins/Cameras/GameplayCameras/Source/GameplayCameras/Private/Nodes/Collision/CollisionPushCameraNode.cpp
--- a/Engine/Plugins/Cameras/GameplayCameras/Source/GameplayCameras/Private/Nodes/Collision/CollisionPushCameraNode.cpp
+++ b/Engine/Plugins/Cameras/GameplayCameras/Source/GameplayCameras/Private/Nodes/Collision/CollisionPushCameraNode.cpp
@@ -4,10 +4,7 @@
 
 #include "CollisionQueryParams.h"
 #include "Core/CameraEvaluationContext.h"
-#include "Core/CameraParameterReader.h"
 #include "Core/CameraSystemEvaluator.h"
-#include "Core/CameraValueInterpolator.h"
-#include "Core/CameraVariableReferenceReader.h"
 #include "Debug/CameraDebugBlock.h"
 #include "Debug/CameraDebugBlockBuilder.h"
 #include "Debug/CameraDebugColors.h"
@@ -27,59 +24,6 @@
 namespace UE::Cameras
 {
 
-class FCollisionPushCameraNodeEvaluator : public FCameraNodeEvaluator
-{
-	UE_DECLARE_CAMERA_NODE_EVALUATOR(GAMEPLAYCAMERAS_API, FCollisionPushCameraNodeEvaluator)
-
-protected:
-
-	// FCameraNodeEvaluator interface.
-	virtual void OnInitialize(const FCameraNodeEvaluatorInitializeParams& Params, FCameraNodeEvaluationResult& OutResult) override;
-	virtual void OnRun(const FCameraNodeEvaluationParams& Params, FCameraNodeEvaluationResult& OutResult) override;
-#if UE_GAMEPLAY_CAMERAS_DEBUG
-	virtual void OnBuildDebugBlocks(const FCameraDebugBlockBuildParams& Params, FCameraDebugBlockBuilder& Builder) override;
-#endif  // UE_GAMEPLAY_CAMERAS_DEBUG
-
-private:
-
-	TOptional<FVector3d> GetFinalSafePosition(const FCameraNodeEvaluationParams& Params, const FCameraNodeEvaluationResult& OutResult);
-	TOptional<FVector3d> GetSafePosition(const FCameraNodeEvaluationParams& Params, const FCameraNodeEvaluationResult& OutResult);
-
-	void RunCollisionTrace(UWorld* World, APlayerController* PlayerController, const FVector3d& SafePosition, const FCameraNodeEvaluationParams& Params, FCameraNodeEvaluationResult& OutResult);
-	void HandleAsyncCollisionTraceResult(UWorld* World, const FVector3d& SafePosition, const FCameraNodeEvaluationParams& Params, FCameraNodeEvaluationResult& OutResult);
-	void HandleCollisionTraceResult(UWorld* World, TArrayView<const FHitResult> HitResults, const FVector3d& SafePosition, const FCameraNodeEvaluationParams& Params, FCameraNodeEvaluationResult& OutResult);
-	void HandleDisabledCollision(const FVector3d& SafePosition, const FCameraNodeEvaluationParams& Params, FCameraNodeEvaluationResult& OutResult);
-	void UpdatePushFactor(bool bFoundHit, float CurrentPushFactor, const FVector3d& SafePosition, const FCameraNodeEvaluationParams& Params, FCameraNodeEvaluationResult& OutResult);
-
-private:
-
-	TCameraVariableReferenceReader<bool> EnableCollisionReader;
-	TCameraVariableReferenceReader<FVector3d> CustomSafePositionReader;
-
-	TCameraParameterReader<float> CollisionSphereRadiusReader;
-	TCameraParameterReader<FVector3d> SafePositionOffsetReader;
-
-	TUniquePtr<FCameraDoubleValueInterpolator> PushInterpolator;
-	TUniquePtr<FCameraDoubleValueInterpolator> PullInterpolator;
-
-	FTraceHandle CollisionTraceHandle;
-
-	float LastPushFactor = 0.f;
-	float LastDampedPushFactor = 0.f;
-
-	enum class ECameraCollisionDirection { Pushing, Pulling };
-	ECameraCollisionDirection LastDirection = ECameraCollisionDirection::Pushing;
-
-#if UE_GAMEPLAY_CAMERAS_DEBUG
-	bool bDebugCollisionEnabled = false;
-	bool bDebugFoundHit = false;
-	bool bDebugGotSafePosition = false;
-	bool bDebugGotSafePositionOffset = false;
-	FString DebugHitObjectName;
-	FVector3d DebugSafePosition;
-#endif
-};
-
 UE_DEFINE_CAMERA_NODE_EVALUATOR(FCollisionPushCameraNodeEvaluator)
 
 UE_DECLARE_CAMERA_DEBUG_BLOCK_START(GAMEPLAYCAMERAS_API, FCollisionPushCameraDebugBlock)
@@ -286,9 +230,6 @@ TOptional<FVector3d> FCollisionPushCameraNodeEvaluator::GetSafePosition(const FC
 
 void FCollisionPushCameraNodeEvaluator::RunCollisionTrace(UWorld* World, APlayerController* PlayerController, const FVector3d& SafePosition, const FCameraNodeEvaluationParams& Params, FCameraNodeEvaluationResult& OutResult)
 {
-	static FName CollisionTraceTag(TEXT("CameraCollision"));
-	static FName CollisionTraceOwnerTag(TEXT("CollisionPushCameraNode"));
-
 	const UCollisionPushCameraNode* CollisionPushNode = GetCameraNodeAs<UCollisionPushCameraNode>();
 	ECollisionChannel CollisionChannel = CollisionPushNode->CollisionChannel;
 
@@ -304,36 +245,57 @@ void FCollisionPushCameraNodeEvaluator::RunCollisionTrace(UWorld* World, APlayer
 		return;
 	}
 
-	FCollisionShape SweepShape = FCollisionShape::MakeSphere(CollisionSphereRadius);
+	FCollisionTraceParams TraceParams;
+	TraceParams.World = World;
+	TraceParams.Pawn = PlayerController->GetPawn();
+	TraceParams.SafePosition = SafePosition;
+	TraceParams.TraceStart = TraceStart;
+	TraceParams.TraceEnd = TraceEnd;
+	TraceParams.CollisionSphereRadius = CollisionSphereRadius;
+	TraceParams.CollisionChannel = CollisionChannel;
+	TraceParams.bRequestedAsyncCollision = CollisionPushNode->bRunAsyncCollision;
+
+	FCollisionTraceResult TraceResult;
+
+	RunCollisionTrace(Params, TraceParams, OutResult, TraceResult);
+
+	CollisionTraceHandle = TraceResult.AsyncTraceHandle;
+	if (!TraceResult.HitResults.IsEmpty())
+	{
+		HandleCollisionTraceResult(World, TraceResult.HitResults, SafePosition, Params, OutResult);
+	}
+}
+
+void FCollisionPushCameraNodeEvaluator::RunCollisionTrace(const FCameraNodeEvaluationParams& Params, const FCollisionTraceParams& TraceParams, FCameraNodeEvaluationResult& OutResult, FCollisionTraceResult& OutTraceResult)
+{
+	static FName CollisionTraceTag(TEXT("CameraCollision"));
+	static FName CollisionTraceOwnerTag(TEXT("CollisionPushCameraNode"));
+
+	FCollisionShape SweepShape = FCollisionShape::MakeSphere(TraceParams.CollisionSphereRadius);
 	// Ignore the player pawn by default.
-	APawn* Pawn = PlayerController->GetPawn();
-	FCollisionQueryParams QueryParams(SCENE_QUERY_STAT(StartCollisionSweep), false, Pawn);
+	FCollisionQueryParams QueryParams(SCENE_QUERY_STAT(StartCollisionSweep), false, TraceParams.Pawn);
 	QueryParams.TraceTag = CollisionTraceTag;
 	QueryParams.OwnerTag = CollisionTraceOwnerTag;
 
-	if (CollisionPushNode->bRunAsyncCollision)
+	if (TraceParams.bRequestedAsyncCollision)
 	{
-		CollisionTraceHandle = World->AsyncSweepByChannel(
+		OutTraceResult.AsyncTraceHandle = TraceParams.World->AsyncSweepByChannel(
 			EAsyncTraceType::Single,
-			TraceStart, TraceEnd, FQuat::Identity,
-			CollisionChannel,
+			TraceParams.TraceStart, TraceParams.TraceEnd, FQuat::Identity,
+			TraceParams.CollisionChannel,
 			SweepShape,
 			QueryParams,
 			FCollisionResponseParams::DefaultResponseParam);
 	}
 	else
 	{
-		TArray<FHitResult> HitResults;
-		World->SweepMultiByChannel(
-			HitResults,
-			TraceStart, TraceEnd, FQuat::Identity,
-			CollisionChannel,
+		TraceParams.World->SweepMultiByChannel(
+			OutTraceResult.HitResults,
+			TraceParams.TraceStart, TraceParams.TraceEnd, FQuat::Identity,
+			TraceParams.CollisionChannel,
 			SweepShape,
 			QueryParams,
 			FCollisionResponseParams::DefaultResponseParam);
-
-		// Handle results right away.
-		HandleCollisionTraceResult(World, HitResults, SafePosition, Params, OutResult);
 	}
 }
 
diff --git a/Engine/Plugins/Cameras/GameplayCameras/Source/GameplayCameras/Private/Nodes/Common/BoomArmCameraNode.cpp b/Engine/Plugins/Cameras/GameplayCameras/Source/GameplayCameras/Private/Nodes/Common/BoomArmCameraNode.cpp
--- a/Engine/Plugins/Cameras/GameplayCameras/Source/GameplayCameras/Private/Nodes/Common/BoomArmCameraNode.cpp
+++ b/Engine/Plugins/Cameras/GameplayCameras/Source/GameplayCameras/Private/Nodes/Common/BoomArmCameraNode.cpp
@@ -4,7 +4,6 @@
 
 #include "Core/CameraEvaluationContext.h"
 #include "Core/CameraOperation.h"
-#include "Core/CameraParameterReader.h"
 #include "Core/CameraRigJoints.h"
 #include "Core/CameraValueInterpolator.h"
 #include "Debug/CameraDebugBlock.h"
@@ -19,44 +18,6 @@
 namespace UE::Cameras
 {
 
-class FBoomArmCameraNodeEvaluator : public FCameraNodeEvaluator
-{
-	UE_DECLARE_CAMERA_NODE_EVALUATOR(GAMEPLAYCAMERAS_API, FBoomArmCameraNodeEvaluator)
-
-protected:
-
-	// FCameraNodeEvaluator interface.
-	virtual void OnBuild(const FCameraNodeEvaluatorBuildParams& Params) override;
-	virtual void OnInitialize(const FCameraNodeEvaluatorInitializeParams& Params, FCameraNodeEvaluationResult& OutResult) override;
-	virtual FCameraNodeEvaluatorChildrenView OnGetChildren() override;
-	virtual void OnRun(const FCameraNodeEvaluationParams& Params, FCameraNodeEvaluationResult& OutResult) override;
-	virtual void OnExecuteOperation(const FCameraOperationParams& Params, FCameraOperation& Operation) override;
-
-#if UE_GAMEPLAY_CAMERAS_DEBUG
-	virtual void OnBuildDebugBlocks(const FCameraDebugBlockBuildParams& Params, FCameraDebugBlockBuilder& Builder) override;
-#endif  // UE_GAMEPLAY_CAMERAS_DEBUG
-
-private:
-
-	APlayerController* GetPlayerController(TSharedPtr<const FCameraEvaluationContext> EvaluationContext) const;
-
-private:
-
-	TCameraParameterReader<FVector3d> BoomOffsetReader;
-	FInput2DCameraNodeEvaluator* InputSlotEvaluator = nullptr;
-	
-	TUniquePtr<TCameraValueInterpolator<double>> BoomLengthInterpolator;
-	TCameraParameterReader<double> MaxForwardInterpolationFactorReader;
-	TCameraParameterReader<double> MaxBackwardInterpolationFactorReader;
-	FVector3d LastPivotLocation;
-	double CumulativePull;
-
-#if UE_GAMEPLAY_CAMERAS_DEBUG
-	FVector2d DebugYawPitch;
-	bool bDebugDidClampPull;
-#endif  // UE_GAMEPLAY_CAMERAS_DEBUG
-};
-
 UE_DEFINE_CAMERA_NODE_EVALUATOR(FBoomArmCameraNodeEvaluator)
 
 UE_DECLARE_CAMERA_DEBUG_BLOCK_START(GAMEPLAYCAMERAS_API, FBoomArmCameraDebugBlock)
@@ -96,7 +57,7 @@ FCameraNodeEvaluatorChildrenView FBoomArmCameraNodeEvaluator::OnGetChildren()
 	return FCameraNodeEvaluatorChildrenView({ InputSlotEvaluator });
 }
 
-void FBoomArmCameraNodeEvaluator::OnRun(const FCameraNodeEvaluationParams& Params, FCameraNodeEvaluationResult& OutResult)
+FRotator FBoomArmCameraNodeEvaluator::ComputeBoomRotation(const FCameraNodeEvaluationParams& Params, FCameraNodeEvaluationResult& OutResult)
 {
 	FRotator3d BoomRotation = FRotator3d::ZeroRotator;
 	if (InputSlotEvaluator)
@@ -105,13 +66,17 @@ void FBoomArmCameraNodeEvaluator::OnRun(const FCameraNodeEvaluationParams& Param
 		const FVector2d YawPitch = InputSlotEvaluator->GetInputValue();
 		BoomRotation = FRotator3d(YawPitch.Y, YawPitch.X, 0);
 	}
-	else if (APlayerController* PlayerController = GetPlayerController(Params.EvaluationContext))
+	else if (APlayerController* PlayerController = Params.EvaluationContext->GetPlayerController())
 	{
 		const FRotator3d ControlRotation = PlayerController->GetControlRotation();
 		BoomRotation = ControlRotation;
 	}
+	return BoomRotation;
+}
 
-	const UBoomArmCameraNode* BoomArmNode = GetCameraNodeAs<UBoomArmCameraNode>();
+void FBoomArmCameraNodeEvaluator::OnRun(const FCameraNodeEvaluationParams& Params, FCameraNodeEvaluationResult& OutResult)
+{
+	const FRotator3d BoomRotation = ComputeBoomRotation(Params, OutResult);
 
 	// Here we want to logically apply transform in this order:
 	//
@@ -217,7 +182,7 @@ void FBoomArmCameraNodeEvaluator::OnExecuteOperation(const FCameraOperationParam
 		// some operations by affecting that pawn rotation ourselves.
 		if (FYawPitchCameraOperation* Op = Operation.CastOperation<FYawPitchCameraOperation>())
 		{
-			if (APlayerController* PlayerController = GetPlayerController(Params.EvaluationContext))
+			if (APlayerController* PlayerController = Params.EvaluationContext->GetPlayerController())
 			{
 				FRotator3d ControlRotation = PlayerController->GetControlRotation();
 				ControlRotation.Yaw = Op->Yaw.Apply(ControlRotation.Yaw);
@@ -228,15 +193,6 @@ void FBoomArmCameraNodeEvaluator::OnExecuteOperation(const FCameraOperationParam
 	}
 }
 
-APlayerController* FBoomArmCameraNodeEvaluator::GetPlayerController(TSharedPtr<const FCameraEvaluationContext> EvaluationContext) const
-{
-	if (EvaluationContext)
-	{
-		return EvaluationContext->GetPlayerController();
-	}
-	return nullptr;
-}
-
 #if UE_GAMEPLAY_CAMERAS_DEBUG
 
 void FBoomArmCameraNodeEvaluator::OnBuildDebugBlocks(const FCameraDebugBlockBuildParams& Params, FCameraDebugBlockBuilder& Builder)
diff --git a/Engine/Plugins/Cameras/GameplayCameras/Source/GameplayCameras/Public/Nodes/Collision/CollisionPushCameraNode.h b/Engine/Plugins/Cameras/GameplayCameras/Source/GameplayCameras/Public/Nodes/Collision/CollisionPushCameraNode.h
--- a/Engine/Plugins/Cameras/GameplayCameras/Source/GameplayCameras/Public/Nodes/Collision/CollisionPushCameraNode.h
+++ b/Engine/Plugins/Cameras/GameplayCameras/Source/GameplayCameras/Public/Nodes/Collision/CollisionPushCameraNode.h
@@ -4,13 +4,16 @@
 
 #include "Core/CameraNode.h"
 #include "Core/CameraParameters.h"
+#include "Core/CameraParameterReader.h"
+#include "Core/CameraValueInterpolator.h"
 #include "Core/CameraVariableReferences.h"
-#include "Nodes/CameraNodeTypes.h"
+#include "Core/CameraVariableReferenceReader.h"
 #include "Engine/EngineTypes.h"
+#include "WorldCollision.h"
 
 #include "CollisionPushCameraNode.generated.h"
 
-class UCameraValueInterpolator;
+#define UE_API GAMEPLAYCAMERAS_API
 
 /**
  * Specifies how to compute the default safe position for the collision camera node
@@ -57,6 +60,99 @@ enum class ECollisionSafePositionOffsetSpace : uint8
 	Pawn
 };
 
+namespace UE::Cameras
+{
+
+/**
+ * The collision push node evaluator class.
+ */
+class FCollisionPushCameraNodeEvaluator : public FCameraNodeEvaluator
+{
+	UE_DECLARE_CAMERA_NODE_EVALUATOR(UE_API, FCollisionPushCameraNodeEvaluator)
+
+protected:
+
+	// FCameraNodeEvaluator interface.
+	UE_API virtual void OnInitialize(const FCameraNodeEvaluatorInitializeParams& Params, FCameraNodeEvaluationResult& OutResult) override;
+	UE_API virtual void OnRun(const FCameraNodeEvaluationParams& Params, FCameraNodeEvaluationResult& OutResult) override;
+#if UE_GAMEPLAY_CAMERAS_DEBUG
+	UE_API virtual void OnBuildDebugBlocks(const FCameraDebugBlockBuildParams& Params, FCameraDebugBlockBuilder& Builder) override;
+#endif  // UE_GAMEPLAY_CAMERAS_DEBUG
+
+protected:
+
+	struct FCollisionTraceParams
+	{
+		UWorld* World = nullptr;
+		APawn* Pawn = nullptr;
+		FVector3d SafePosition = FVector3d::ZeroVector;
+		FVector3d TraceStart = FVector3d::ZeroVector;
+		FVector3d TraceEnd = FVector3d::ZeroVector;
+		float CollisionSphereRadius = 1.f;
+		ECollisionChannel CollisionChannel = ECollisionChannel::ECC_Camera;
+		bool bRequestedAsyncCollision = false;
+	};
+
+	struct FCollisionTraceResult
+	{
+		TArray<FHitResult> HitResults;
+		FTraceHandle AsyncTraceHandle;
+	};
+
+	/**
+	 * Run the collision trace using the given parameters.
+	 *
+	 * @param Params  The parameters for the current node evaluation
+	 * @param TraceParams  The parameters for the collision trace
+	 * @param OutResult   The result for the current node evaluation
+	 * @param OutTraceResult  The result with either the handle to the asynchronous trace, if bRequestedAsyncCollision 
+	 *						  was true and it was possible to honor it, or the list of hit results to be processed 
+	 *						  synchronously
+	 */
+	UE_API virtual void RunCollisionTrace(const FCameraNodeEvaluationParams& Params, const FCollisionTraceParams& TraceParams, FCameraNodeEvaluationResult& OutResult, FCollisionTraceResult& OutTraceResult);
+
+private:
+
+	TOptional<FVector3d> GetFinalSafePosition(const FCameraNodeEvaluationParams& Params, const FCameraNodeEvaluationResult& OutResult);
+	TOptional<FVector3d> GetSafePosition(const FCameraNodeEvaluationParams& Params, const FCameraNodeEvaluationResult& OutResult);
+
+	void RunCollisionTrace(UWorld* World, APlayerController* PlayerController, const FVector3d& SafePosition, const FCameraNodeEvaluationParams& Params, FCameraNodeEvaluationResult& OutResult);
+	void HandleAsyncCollisionTraceResult(UWorld* World, const FVector3d& SafePosition, const FCameraNodeEvaluationParams& Params, FCameraNodeEvaluationResult& OutResult);
+	void HandleCollisionTraceResult(UWorld* World, TArrayView<const FHitResult> HitResults, const FVector3d& SafePosition, const FCameraNodeEvaluationParams& Params, FCameraNodeEvaluationResult& OutResult);
+	void HandleDisabledCollision(const FVector3d& SafePosition, const FCameraNodeEvaluationParams& Params, FCameraNodeEvaluationResult& OutResult);
+	void UpdatePushFactor(bool bFoundHit, float CurrentPushFactor, const FVector3d& SafePosition, const FCameraNodeEvaluationParams& Params, FCameraNodeEvaluationResult& OutResult);
+
+private:
+
+	TCameraVariableReferenceReader<bool> EnableCollisionReader;
+	TCameraVariableReferenceReader<FVector3d> CustomSafePositionReader;
+
+	TCameraParameterReader<float> CollisionSphereRadiusReader;
+	TCameraParameterReader<FVector3d> SafePositionOffsetReader;
+
+	TUniquePtr<FCameraDoubleValueInterpolator> PushInterpolator;
+	TUniquePtr<FCameraDoubleValueInterpolator> PullInterpolator;
+
+	FTraceHandle CollisionTraceHandle;
+
+	float LastPushFactor = 0.f;
+	float LastDampedPushFactor = 0.f;
+
+	enum class ECameraCollisionDirection { Pushing, Pulling };
+	ECameraCollisionDirection LastDirection = ECameraCollisionDirection::Pushing;
+
+#if UE_GAMEPLAY_CAMERAS_DEBUG
+	bool bDebugCollisionEnabled = false;
+	bool bDebugFoundHit = false;
+	bool bDebugGotSafePosition = false;
+	bool bDebugGotSafePositionOffset = false;
+	FString DebugHitObjectName;
+	FVector3d DebugSafePosition;
+#endif
+};
+
+} // namespace UE::Cameras
+
 /**
  * A node that pushes the camera towards a "safe position" when it is colliding with 
  * the environment. By default, the "safe position" is the pivot of the camera (if any) 
@@ -121,11 +217,13 @@ public:
 
 public:
 
-	UCollisionPushCameraNode(const FObjectInitializer& ObjectInit);
+	UE_API UCollisionPushCameraNode(const FObjectInitializer& ObjectInit);
 
 protected:
 
 	// UCameraNode interface.
-	virtual FCameraNodeEvaluatorPtr OnBuildEvaluator(FCameraNodeEvaluatorBuilder& Builder) const override;
+	UE_API virtual FCameraNodeEvaluatorPtr OnBuildEvaluator(FCameraNodeEvaluatorBuilder& Builder) const override;
 };
 
+#undef UE_API
+
diff --git a/Engine/Plugins/Cameras/GameplayCameras/Source/GameplayCameras/Public/Nodes/Common/BoomArmCameraNode.h b/Engine/Plugins/Cameras/GameplayCameras/Source/GameplayCameras/Public/Nodes/Common/BoomArmCameraNode.h
--- a/Engine/Plugins/Cameras/GameplayCameras/Source/GameplayCameras/Public/Nodes/Common/BoomArmCameraNode.h
+++ b/Engine/Plugins/Cameras/GameplayCameras/Source/GameplayCameras/Public/Nodes/Common/BoomArmCameraNode.h
@@ -4,12 +4,64 @@
 
 #include "Core/CameraNode.h"
 #include "Core/CameraParameters.h"
+#include "Core/CameraParameterReader.h"
 
 #include "BoomArmCameraNode.generated.h"
 
+#define UE_API GAMEPLAYCAMERAS_API
+
 class UCameraValueInterpolator;
 class UInput2DCameraNode;
 
+namespace UE::Cameras
+{
+
+class FCameraEvaluationContext;
+class FInput2DCameraNodeEvaluator;
+template<typename T> class TCameraValueInterpolator;
+
+/**
+ * The boom arm node evaluator class.
+ */
+class FBoomArmCameraNodeEvaluator : public FCameraNodeEvaluator
+{
+	UE_DECLARE_CAMERA_NODE_EVALUATOR(UE_API, FBoomArmCameraNodeEvaluator)
+
+protected:
+
+	// FCameraNodeEvaluator interface.
+	UE_API virtual void OnBuild(const FCameraNodeEvaluatorBuildParams& Params) override;
+	UE_API virtual void OnInitialize(const FCameraNodeEvaluatorInitializeParams& Params, FCameraNodeEvaluationResult& OutResult) override;
+	UE_API virtual FCameraNodeEvaluatorChildrenView OnGetChildren() override;
+	UE_API virtual void OnRun(const FCameraNodeEvaluationParams& Params, FCameraNodeEvaluationResult& OutResult) override;
+	UE_API virtual void OnExecuteOperation(const FCameraOperationParams& Params, FCameraOperation& Operation) override;
+
+#if UE_GAMEPLAY_CAMERAS_DEBUG
+	UE_API virtual void OnBuildDebugBlocks(const FCameraDebugBlockBuildParams& Params, FCameraDebugBlockBuilder& Builder) override;
+#endif  // UE_GAMEPLAY_CAMERAS_DEBUG
+
+	// FBoomArmCameraNodeEvaluator interface.
+	UE_API virtual FRotator ComputeBoomRotation(const FCameraNodeEvaluationParams& Params, FCameraNodeEvaluationResult& OutResult);
+
+private:
+
+	TCameraParameterReader<FVector3d> BoomOffsetReader;
+	FInput2DCameraNodeEvaluator* InputSlotEvaluator = nullptr;
+
+	TUniquePtr<TCameraValueInterpolator<double>> BoomLengthInterpolator;
+	TCameraParameterReader<double> MaxForwardInterpolationFactorReader;
+	TCameraParameterReader<double> MaxBackwardInterpolationFactorReader;
+	FVector3d LastPivotLocation;
+	double CumulativePull;
+
+#if UE_GAMEPLAY_CAMERAS_DEBUG
+	FVector2d DebugYawPitch;
+	bool bDebugDidClampPull;
+#endif  // UE_GAMEPLAY_CAMERAS_DEBUG
+};
+
+}  // namespace UE::Cameras
+
 /**
  * A camera node that can rotate the camera in yaw and pitch based on player input.
  */
@@ -20,13 +72,13 @@ class UBoomArmCameraNode : public UCameraNode
 
 public:
 
-	UBoomArmCameraNode(const FObjectInitializer& ObjInit);
+	UE_API UBoomArmCameraNode(const FObjectInitializer& ObjInit);
 
 protected:
 
 	// UCameraNode interface.
-	virtual FCameraNodeChildrenView OnGetChildren() override;
-	virtual FCameraNodeEvaluatorPtr OnBuildEvaluator(FCameraNodeEvaluatorBuilder& Builder) const override;
+	UE_API virtual FCameraNodeChildrenView OnGetChildren() override;
+	UE_API virtual FCameraNodeEvaluatorPtr OnBuildEvaluator(FCameraNodeEvaluatorBuilder& Builder) const override;
 
 public:
 
@@ -60,3 +112,5 @@ public:
 	TObjectPtr<UInput2DCameraNode> InputSlot;
 };
 
+#undef UE_API
+
