Cameras: add easing blend node

diff --git a/Engine/Plugins/Cameras/GameplayCameras/Source/GameplayCameras/Private/Nodes/EasingBlendCameraNode.cpp b/Engine/Plugins/Cameras/GameplayCameras/Source/GameplayCameras/Private/Nodes/EasingBlendCameraNode.cpp
new file mode 100644
--- /dev/null
+++ b/Engine/Plugins/Cameras/GameplayCameras/Source/GameplayCameras/Private/Nodes/EasingBlendCameraNode.cpp
@@ -0,0 +1,151 @@
+// Copyright Epic Games, Inc. All Rights Reserved.
+
+#include "Nodes/Blends/EasingBlendCameraNode.h"
+
+#include "Math/UnrealMathUtility.h"
+
+#include UE_INLINE_GENERATED_CPP_BY_NAME(EasingBlendCameraNode)
+
+namespace UE::Cameras
+{
+
+namespace Internal
+{
+
+float SinOut(float InTime)
+{
+	return FMath::Sin(.5f * UE_PI * InTime);
+}
+
+float SinIn(float InTime)
+{
+	return 1.f - SinOut(1.f - InTime);
+}
+
+float PowIn(float InTime, float Power)
+{
+	return FMath::Pow(InTime, Power);
+}
+
+float PowOut(float InTime, float Power)
+{
+	return 1.f - FMath::Pow(1.f - InTime, Power);
+}
+
+float ExpIn(float InTime)
+{
+	return FMath::Pow(2, 10*(InTime - 1.f));
+}
+
+float ExpOut(float InTime)
+{
+	return 1.f - ExpIn(1.f - InTime);
+}
+
+float CircIn(float InTime)
+{
+	return 1.f - FMath::Sqrt(1-InTime*InTime);
+}
+
+float CircOut(float InTime)
+{
+	return 1.f - CircIn(1.f - InTime);
+}
+
+float EvaluateEasing(EEasingCameraBlendType BlendType, float Factor)
+{
+	// The algorithms below match the ones from Sequencer (see UMovieSceneBuiltInEasingFunction).
+	const float InTime = Factor * 2.f;
+	const float OutTime = (Factor - .5f) * 2.f;
+
+	switch (BlendType)
+	{
+		case EEasingCameraBlendType::SinIn:
+			return SinIn(Factor);
+		case EEasingCameraBlendType::SinOut:
+			return SinOut(Factor);
+		case EEasingCameraBlendType::SinInOut:
+			return InTime < 1.f ? .5f * SinIn(InTime) : .5f + .5f * SinOut(OutTime);
+
+		case EEasingCameraBlendType::QuadIn:
+			return PowIn(Factor,2);
+		case EEasingCameraBlendType::QuadOut:
+			return PowOut(Factor,2);
+		case EEasingCameraBlendType::QuadInOut:
+			return InTime < 1.f ? .5f * PowIn(InTime, 2) : .5f + .5f * PowOut(OutTime, 2);
+
+		case EEasingCameraBlendType::Cubic:
+			return FMath::Clamp<float>(FMath::CubicInterp<float>(0.f, 0.f, 1.f, 0.f, Factor), 0.f, 1.f);
+		case EEasingCameraBlendType::CubicIn:
+			return PowIn(Factor,3);
+		case EEasingCameraBlendType::CubicOut:
+			return PowOut(Factor,3);
+		case EEasingCameraBlendType::CubicInOut:
+			return InTime < 1.f ? .5f * PowIn(InTime,3) : .5f + .5f * PowOut(OutTime, 3);
+
+		case EEasingCameraBlendType::HermiteCubicInOut:
+			return FMath::Clamp<float>(FMath::SmoothStep(0.0f, 1.0f, Factor), 0.0f, 1.0f);
+
+		case EEasingCameraBlendType::QuartIn:
+			return PowIn(Factor, 4);
+		case EEasingCameraBlendType::QuartOut:
+			return PowOut(Factor, 4);
+		case EEasingCameraBlendType::QuartInOut:
+			return InTime < 1.f ? .5f * PowIn(InTime, 4) : .5f + .5f * PowOut(OutTime, 4);
+
+		case EEasingCameraBlendType::QuintIn:
+			return PowIn(Factor,5);
+		case EEasingCameraBlendType::QuintOut:
+			return PowOut(Factor,5);
+		case EEasingCameraBlendType::QuintInOut:
+			return InTime < 1.f ? .5f * PowIn(InTime, 5) : .5f + .5f * PowOut(OutTime, 5);
+
+		case EEasingCameraBlendType::ExpoIn:
+			return ExpIn(Factor);
+		case EEasingCameraBlendType::ExpoOut:
+			return ExpOut(Factor);
+		case EEasingCameraBlendType::ExpoInOut:
+			return InTime < 1.f ? .5f * ExpIn(InTime) : .5f + .5f * ExpOut(OutTime);
+
+		case EEasingCameraBlendType::CircIn:
+			return CircIn(Factor);
+		case EEasingCameraBlendType::CircOut:
+			return CircOut(Factor);
+		case EEasingCameraBlendType::CircInOut:
+			return InTime < 1.f ? .5f * CircIn(InTime) : .5f + .5f * CircOut(OutTime);
+
+		case EEasingCameraBlendType::Linear:
+		default:
+			return Factor;
+	}
+}
+
+}  // namespace Internal
+
+class FEasingBlendCameraNodeEvaluator : public FSimpleFixedTimeBlendCameraNodeEvaluator
+{
+	UE_DECLARE_BLEND_CAMERA_NODE_EVALUATOR_EX(GAMEPLAYCAMERAS_API, FEasingBlendCameraNodeEvaluator, FSimpleFixedTimeBlendCameraNodeEvaluator)
+
+protected:
+	virtual void OnComputeBlendFactor(const FCameraNodeEvaluationParams& Params, FSimpleBlendCameraNodeEvaluationResult& OutResult) override;
+};
+
+UE_DEFINE_BLEND_CAMERA_NODE_EVALUATOR(FEasingBlendCameraNodeEvaluator)
+
+void FEasingBlendCameraNodeEvaluator::OnComputeBlendFactor(const FCameraNodeEvaluationParams& Params, FSimpleBlendCameraNodeEvaluationResult& OutResult)
+{
+	using namespace UE::Cameras::Internal;
+
+	const UEasingBlendCameraNode* BlendNode = GetCameraNodeAs<UEasingBlendCameraNode>();
+	const float TimeFactor = GetTimeFactor();
+	OutResult.BlendFactor = EvaluateEasing(BlendNode->BlendType, TimeFactor);
+}
+
+}  // namespace UE::Cameras
+
+FCameraNodeEvaluatorPtr UEasingBlendCameraNode::OnBuildEvaluator(FCameraNodeEvaluatorBuilder& Builder) const
+{
+	using namespace UE::Cameras;
+	return Builder.BuildEvaluator<FEasingBlendCameraNodeEvaluator>();
+}
+
diff --git a/Engine/Plugins/Cameras/GameplayCameras/Source/GameplayCameras/Public/Nodes/Blends/EasingBlendCameraNode.h b/Engine/Plugins/Cameras/GameplayCameras/Source/GameplayCameras/Public/Nodes/Blends/EasingBlendCameraNode.h
new file mode 100644
--- /dev/null
+++ b/Engine/Plugins/Cameras/GameplayCameras/Source/GameplayCameras/Public/Nodes/Blends/EasingBlendCameraNode.h
@@ -0,0 +1,74 @@
+// Copyright Epic Games, Inc. All Rights Reserved.
+
+#pragma once
+
+#include "Nodes/Blends/SimpleBlendCameraNode.h"
+
+#include "EasingBlendCameraNode.generated.h"
+
+/**
+ * The easing type.
+ *
+ * These types are copied from, and behave in the same way as the corresponding easing
+ * types from Sequencer.
+ */
+UENUM()
+enum class EEasingCameraBlendType
+{
+	// Linear easing
+	Linear UMETA(Grouping=Linear, DisplayName="Linear"),
+	// Sinusoidal easing
+	SinIn UMETA(Grouping=Sinusoidal, DisplayName="Sinusoidal In"),
+	SinOut UMETA(Grouping=Sinusoidal, DisplayName="Sinusoidal Out"),
+	SinInOut UMETA(Grouping=Sinusoidal, DisplayName="Sinusoidal InOut"),
+	// Quadratic easing
+	QuadIn UMETA(Grouping=Quadratic, DisplayName="Quadratic In"),
+	QuadOut UMETA(Grouping=Quadratic, DisplayName="Quadratic Out"),
+	QuadInOut UMETA(Grouping=Quadratic, DisplayName="Quadratic InOut"),
+	// Cubic easing
+	Cubic UMETA(Grouping=Cubic, DisplayName="Cubic"),
+	CubicIn UMETA(Grouping=Cubic, DisplayName="Cubic In"),
+	CubicOut UMETA(Grouping=Cubic, DisplayName="Cubic Out"),
+	CubicInOut UMETA(Grouping=Cubic, DisplayName="Cubic InOut"),
+	HermiteCubicInOut UMETA(Grouping=Cubic, DisplayName="Hermite-Cubic InOut"),
+	// Quartic easing
+	QuartIn UMETA(Grouping=Quartic, DisplayName="Quartic In"),
+	QuartOut UMETA(Grouping=Quartic, DisplayName="Quartic Out"),
+	QuartInOut UMETA(Grouping=Quartic, DisplayName="Quartic InOut"),
+	// Quintic easing
+	QuintIn UMETA(Grouping=Quintic, DisplayName="Quintic In"),
+	QuintOut UMETA(Grouping=Quintic, DisplayName="Quintic Out"),
+	QuintInOut UMETA(Grouping=Quintic, DisplayName="Quintic InOut"),
+	// Exponential easing
+	ExpoIn UMETA(Grouping=Exponential, DisplayName="Exponential In"),
+	ExpoOut UMETA(Grouping=Exponential, DisplayName="Exponential Out"),
+	ExpoInOut UMETA(Grouping=Exponential, DisplayName="Exponential InOut"),
+	// Circular easing
+	CircIn UMETA(Grouping=Circular, DisplayName="Circular In"),
+	CircOut UMETA(Grouping=Circular, DisplayName="Circular Out"),
+	CircInOut UMETA(Grouping=Circular, DisplayName="Circular InOut"),
+};
+
+/**
+ * A blend camera node that implements basic easing curve types.
+ */
+UCLASS(MinimalAPI)
+class UEasingBlendCameraNode : public USimpleFixedTimeBlendCameraNode
+{
+	GENERATED_BODY()
+
+public:
+
+	void SetCameraBlendType(EEasingCameraBlendType BlendTypeIn) { BlendType = BlendTypeIn; }
+
+protected:
+
+	virtual FCameraNodeEvaluatorPtr OnBuildEvaluator(FCameraNodeEvaluatorBuilder& Builder) const override;
+	
+public:
+
+	/** The type of curve to use. */
+	UPROPERTY(EditAnywhere, Category=Blending)
+	EEasingCameraBlendType BlendType = EEasingCameraBlendType::Linear;
+};
+
diff --git a/Engine/Plugins/Cameras/GameplayCameras/Source/GameplayCameras/Public/Nodes/Blends/SmoothBlendCameraNode.h b/Engine/Plugins/Cameras/GameplayCameras/Source/GameplayCameras/Public/Nodes/Blends/SmoothBlendCameraNode.h
--- a/Engine/Plugins/Cameras/GameplayCameras/Source/GameplayCameras/Public/Nodes/Blends/SmoothBlendCameraNode.h
+++ b/Engine/Plugins/Cameras/GameplayCameras/Source/GameplayCameras/Public/Nodes/Blends/SmoothBlendCameraNode.h
@@ -17,7 +17,7 @@ enum class ESmoothCameraBlendType
 };
 
 /**
- * A blend camera rig that implements the smoothstep and smoothersteps algorithms.
+ * A blend camera node that implements the smoothstep and smoothersteps algorithms.
  */
 UCLASS(MinimalAPI)
 class USmoothBlendCameraNode : public USimpleFixedTimeBlendCameraNode
