Cameras: look for rigs and rig proxies used by directors more than one package away

This fixes the issue of, say, a Blueprint director using some data table or chooser asset for selecting rigs, which meant that the "family shortcuts" and other things using the UsageInfo API had incomplete info.

diff --git a/Engine/Plugins/Cameras/GameplayCameras/Source/GameplayCameras/Private/Directors/BlueprintCameraDirector.cpp b/Engine/Plugins/Cameras/GameplayCameras/Source/GameplayCameras/Private/Directors/BlueprintCameraDirector.cpp
--- a/Engine/Plugins/Cameras/GameplayCameras/Source/GameplayCameras/Private/Directors/BlueprintCameraDirector.cpp
+++ b/Engine/Plugins/Cameras/GameplayCameras/Source/GameplayCameras/Private/Directors/BlueprintCameraDirector.cpp
@@ -17,6 +17,7 @@
 #include "GameFramework/Actor.h"
 #include "GameFramework/ControllerGameplayCameraEvaluationComponent.h"
 #include "GameplayCameras.h"
+#include "GameplayCamerasSettings.h"
 #include "Helpers/OutgoingReferenceFinder.h"
 
 #include UE_INLINE_GENERATED_CPP_BY_NAME(BlueprintCameraDirector)
@@ -531,6 +532,7 @@ void UBlueprintCameraDirector::OnGatherRigUsageInfo(FCameraDirectorRigUsageInfo&
 
 	TArray<UClass*> RefClasses { UCameraRigAsset::StaticClass(), UCameraRigProxyAsset::StaticClass() };
 	FOutgoingReferenceFinder ReferenceFinder(EvaluatorBlueprint, RefClasses);
+	ReferenceFinder.SetMaxDistance(GetDefault<UGameplayCamerasSettings>()->MaxUsageSearchDistance);
 	ReferenceFinder.CollectReferences();
 	ReferenceFinder.GetReferencesOfClass<UCameraRigAsset>(UsageInfo.CameraRigs);
 	ReferenceFinder.GetReferencesOfClass<UCameraRigProxyAsset>(UsageInfo.CameraRigProxies);
diff --git a/Engine/Plugins/Cameras/GameplayCameras/Source/GameplayCameras/Private/Directors/StateTreeCameraDirector.cpp b/Engine/Plugins/Cameras/GameplayCameras/Source/GameplayCameras/Private/Directors/StateTreeCameraDirector.cpp
--- a/Engine/Plugins/Cameras/GameplayCameras/Source/GameplayCameras/Private/Directors/StateTreeCameraDirector.cpp
+++ b/Engine/Plugins/Cameras/GameplayCameras/Source/GameplayCameras/Private/Directors/StateTreeCameraDirector.cpp
@@ -9,6 +9,7 @@
 #include "Core/CameraRigProxyAsset.h"
 #include "Directors/CameraDirectorStateTreeSchema.h"
 #include "GameplayCameras.h"
+#include "GameplayCamerasSettings.h"
 #include "Helpers/OutgoingReferenceFinder.h"
 #include "Logging/TokenizedMessage.h"
 #include "StateTree.h"
@@ -245,6 +246,7 @@ void UStateTreeCameraDirector::OnGatherRigUsageInfo(FCameraDirectorRigUsageInfo&
 
 	TArray<UClass*> RefClasses { UCameraRigAsset::StaticClass(), UCameraRigProxyAsset::StaticClass() };
 	FOutgoingReferenceFinder ReferenceFinder(const_cast<UStateTree*>(StateTree), RefClasses);
+	ReferenceFinder.SetMaxDistance(GetDefault<UGameplayCamerasSettings>()->MaxUsageSearchDistance);
 	ReferenceFinder.CollectReferences();
 	ReferenceFinder.GetReferencesOfClass<UCameraRigAsset>(UsageInfo.CameraRigs);
 	ReferenceFinder.GetReferencesOfClass<UCameraRigProxyAsset>(UsageInfo.CameraRigProxies);
diff --git a/Engine/Plugins/Cameras/GameplayCameras/Source/GameplayCameras/Private/Helpers/OutgoingReferenceFinder.cpp b/Engine/Plugins/Cameras/GameplayCameras/Source/GameplayCameras/Private/Helpers/OutgoingReferenceFinder.cpp
--- a/Engine/Plugins/Cameras/GameplayCameras/Source/GameplayCameras/Private/Helpers/OutgoingReferenceFinder.cpp
+++ b/Engine/Plugins/Cameras/GameplayCameras/Source/GameplayCameras/Private/Helpers/OutgoingReferenceFinder.cpp
@@ -36,19 +36,27 @@ void FOutgoingReferenceFinder::Initialize(UObject* InRootObject)
 	ArShouldSkipBulkData = true;
 }
 
+void FOutgoingReferenceFinder::SetMaxDistance(int32 InMaxDistance)
+{
+	MaxDistance = FMath::Clamp(InMaxDistance, 0, 5);
+}
+
 void FOutgoingReferenceFinder::CollectReferences()
 {
 	ObjectsToVisit.Reset();
 	VisitedObjects.Reset();
 
-	ObjectsToVisit.Add(RootObject);
+	ObjectsToVisit.Add({ RootObject, PackageScope, 0 });
 	while (ObjectsToVisit.Num() > 0)
 	{
-		UObject* CurObj = ObjectsToVisit.Pop(EAllowShrinking::No);
-		if (!VisitedObjects.Contains(CurObj))
+		FObjectToVisit CurObj = ObjectsToVisit.Pop(EAllowShrinking::No);
+		if (!VisitedObjects.Contains(CurObj.Obj))
 		{
-			VisitedObjects.Add(CurObj);
-			CurObj->Serialize(*this);
+			VisitedObjects.Add(CurObj.Obj);
+
+			SerializeState = { CurObj.Package, CurObj.Distance };
+			CurObj.Obj->Serialize(*this);
+			SerializeState = FSerializeState();
 		}
 	}
 }
@@ -78,9 +86,18 @@ FArchive& FOutgoingReferenceFinder::operator<<(UObject*& ObjRef)
 			ReferencesOfClass.Add(ObjRef);
 		}
 
-		if (ObjRef->IsIn(PackageScope) && !VisitedObjects.Contains(ObjRef))
+		if (!VisitedObjects.Contains(ObjRef))
 		{
-			ObjectsToVisit.Add(ObjRef);
+			UPackage* ObjRefPackage = ObjRef->GetOutermost();
+			// Note that this distance is maybe not the "correct" one. That is: there may be another path to ObjRef that 
+			// has a lower distance. But that's OK. Either we queue a visit of ObjRef right now, and we'll ignore the
+			// duplicate with the shorter distance, or the current distance is over the threshold and we skip visiting
+			// it now, but we'll visit it when we encounter it again with a distance under the threshold.
+			const int32 ObjRefDistance = SerializeState.CurrentDistance + (ObjRefPackage == SerializeState.CurrentPackage ? 0 : 1);
+			if (ObjRefDistance <= MaxDistance)
+			{
+				ObjectsToVisit.Add({ ObjRef, ObjRefPackage, ObjRefDistance });
+			}
 		}
 	}
 	return *this;
diff --git a/Engine/Plugins/Cameras/GameplayCameras/Source/GameplayCameras/Public/GameplayCamerasSettings.h b/Engine/Plugins/Cameras/GameplayCameras/Source/GameplayCameras/Public/GameplayCamerasSettings.h
--- a/Engine/Plugins/Cameras/GameplayCameras/Source/GameplayCameras/Public/GameplayCamerasSettings.h
+++ b/Engine/Plugins/Cameras/GameplayCameras/Source/GameplayCameras/Public/GameplayCamerasSettings.h
@@ -44,6 +44,14 @@ public:
 	UPROPERTY(EditAnywhere, Config, Category="General")
 	int32 CombinedCameraRigNumThreshold = 10;
 
+public:
+
+	/**
+	 * How many packages away to look for references to camera rigs and rig proxies from a camera director.
+	 */
+	UPROPERTY(EditAnywhere, Config, Category="Camera Director", meta=(ClampMin=1, ClampMax=5))
+	int32 MaxUsageSearchDistance = 2;
+
 public:
 
 	/** The default angle tolerance to accept an aiming operation. */
diff --git a/Engine/Plugins/Cameras/GameplayCameras/Source/GameplayCameras/Public/Helpers/OutgoingReferenceFinder.h b/Engine/Plugins/Cameras/GameplayCameras/Source/GameplayCameras/Public/Helpers/OutgoingReferenceFinder.h
--- a/Engine/Plugins/Cameras/GameplayCameras/Source/GameplayCameras/Public/Helpers/OutgoingReferenceFinder.h
+++ b/Engine/Plugins/Cameras/GameplayCameras/Source/GameplayCameras/Public/Helpers/OutgoingReferenceFinder.h
@@ -27,6 +27,18 @@ public:
 	/** Creates a new instance of the reference finder. */
 	UE_API FOutgoingReferenceFinder(UObject* InRootObject, TArrayView<UClass*> InReferencedObjectClasses);
 
+	/**
+	 * Sets the maximum distance (in number of packages) to search for references.
+	 * Zero means only search the root object package.
+	 * One means search directly referenced packages.
+	 * Higher distances exponentially grow the number of objects to search, so the maximum
+	 * value settable is 5.
+	 */
+	UE_API void SetMaxDistance(int32 InMaxDistance);
+
+	/** Gets the maximum distance (in number of packages) to search for references. */
+	int32 GetMaxDistance() const { return MaxDistance; }
+
 	/** Runs the reference finding. */
 	UE_API void CollectReferences();
 
@@ -50,13 +62,28 @@ private:
 
 private:
 
+	struct FObjectToVisit
+	{
+		UObject* Obj = nullptr;
+		UPackage* Package = nullptr;
+		int32 Distance = 0;
+	};
+
+	struct FSerializeState
+	{
+		UPackage* CurrentPackage = nullptr;
+		int32 CurrentDistance = 0;
+	};
+
 	UObject* RootObject;
 	UPackage* PackageScope;
+	int32 MaxDistance = 1;
 
 	TSet<UClass*> TargetObjectClasses;
 
-	TArray<UObject*> ObjectsToVisit;
+	TArray<FObjectToVisit> ObjectsToVisit;
 	TSet<UObject*> VisitedObjects;
+	FSerializeState SerializeState;
 
 	TMap<UClass*, TSet<UObject*>> ReferencedObjects;
 };
