Cameras: fix logging formatting

diff --git a/Engine/Plugins/Cameras/GameplayCameras/Source/GameplayCameras/Private/GameFramework/GameplayCameraComponentBase.cpp b/Engine/Plugins/Cameras/GameplayCameras/Source/GameplayCameras/Private/GameFramework/GameplayCameraComponentBase.cpp
--- a/Engine/Plugins/Cameras/GameplayCameras/Source/GameplayCameras/Private/GameFramework/GameplayCameraComponentBase.cpp
+++ b/Engine/Plugins/Cameras/GameplayCameras/Source/GameplayCameras/Private/GameFramework/GameplayCameraComponentBase.cpp
@@ -214,7 +214,7 @@ void UGameplayCameraComponentBase::EnsureEvaluationContext(APlayerController* Pl
 		FFrame::KismetExecutionMessage(
 				*FString::Format(
 					TEXT("Trying to create an evaluation context for Player Controller '{0}' on Gameplay Camera component '{1}.{2}' "
-						 "but it was already created for Player Controller '{1}'. The evaluation context will be re-created."),
+						 "but it was already created for Player Controller '{3}'. The evaluation context will be re-created."),
 					{ 
 						*GetNameSafe(PlayerController), 
 						*GetNameSafe(GetOwner()),
diff --git a/Engine/Plugins/Cameras/GameplayCameras/Source/GameplayCameras/Private/GameFramework/GameplayCamerasPlayerCameraManager.cpp b/Engine/Plugins/Cameras/GameplayCameras/Source/GameplayCameras/Private/GameFramework/GameplayCamerasPlayerCameraManager.cpp
--- a/Engine/Plugins/Cameras/GameplayCameras/Source/GameplayCameras/Private/GameFramework/GameplayCamerasPlayerCameraManager.cpp
+++ b/Engine/Plugins/Cameras/GameplayCameras/Source/GameplayCameras/Private/GameFramework/GameplayCamerasPlayerCameraManager.cpp
@@ -149,10 +149,12 @@ void AGameplayCamerasPlayerCameraManager::ActivateGameplayCamera(UGameplayCamera
 					const bool bSuccess = ActiveContext->AddChildContext(GameplayCameraContext.ToSharedRef());
 					if (!bSuccess)
 					{
+						UObject* ActiveContextOwner = ActiveContext->GetOwner();
 						UE_LOG(LogCameraSystem, Error,
-								TEXT("Couldn't insert camera director for '{0}.{1}' inside camera director for '{2}.{3}'. "
+								TEXT("Couldn't insert camera director for '%s.%s' inside camera director for '%s'. "
 									 "Activation failed."),
-								*GetNameSafe(GameplayCamera->GetOwner()), *GetNameSafe(GameplayCamera));
+								*GetNameSafe(GameplayCamera->GetOwner()), *GetNameSafe(GameplayCamera),
+								*GetNameSafe(ActiveContextOwner));
 					}
 				}
 				else
