Cameras: fix possibly incorrect evaluation layer when starting a modifier rig

diff --git a/Engine/Plugins/Cameras/GameplayCameras/Source/GameplayCameras/Private/GameFramework/IGameplayCameraSystemHost.cpp b/Engine/Plugins/Cameras/GameplayCameras/Source/GameplayCameras/Private/GameFramework/IGameplayCameraSystemHost.cpp
--- a/Engine/Plugins/Cameras/GameplayCameras/Source/GameplayCameras/Private/GameFramework/IGameplayCameraSystemHost.cpp
+++ b/Engine/Plugins/Cameras/GameplayCameras/Source/GameplayCameras/Private/GameFramework/IGameplayCameraSystemHost.cpp
@@ -129,7 +129,7 @@ FCameraRigInstanceID IGameplayCameraSystemHost::StartCameraModifierRig(
 		TSharedPtr<FCameraModifierService> CameraModifierService = CameraSystemEvaluator->FindEvaluationService<FCameraModifierService>();
 		if (ensure(CameraModifierService))
 		{
-			return CameraModifierService->StartCameraModifierRig(CameraRig, EvaluationContext.ToSharedRef(), ECameraRigLayer::Global, OrderKey);
+			return CameraModifierService->StartCameraModifierRig(CameraRig, EvaluationContext.ToSharedRef(), EvaluationLayer, OrderKey);
 		}
 	}
 
