IWYU

diff --git a/Engine/Plugins/Cameras/GameplayCameras/Source/GameplayCameras/Public/Core/CameraObjectStorage.h b/Engine/Plugins/Cameras/GameplayCameras/Source/GameplayCameras/Public/Core/CameraObjectStorage.h
--- a/Engine/Plugins/Cameras/GameplayCameras/Source/GameplayCameras/Public/Core/CameraObjectStorage.h
+++ b/Engine/Plugins/Cameras/GameplayCameras/Source/GameplayCameras/Public/Core/CameraObjectStorage.h
@@ -8,6 +8,7 @@
 #include "Templates/AlignmentTemplates.h"
 #include "Templates/EnableIf.h"
 #include "Templates/PointerIsConvertibleFromTo.h"
+#include "Templates/UnrealTemplate.h"
 
 namespace UE::Cameras
 {
