Cameras: fix how the output camera component is created

diff --git a/Engine/Plugins/Cameras/GameplayCameras/Source/GameplayCameras/Private/GameFramework/GameplayCameraComponentBase.cpp b/Engine/Plugins/Cameras/GameplayCameras/Source/GameplayCameras/Private/GameFramework/GameplayCameraComponentBase.cpp
--- a/Engine/Plugins/Cameras/GameplayCameras/Source/GameplayCameras/Private/GameFramework/GameplayCameraComponentBase.cpp
+++ b/Engine/Plugins/Cameras/GameplayCameras/Source/GameplayCameras/Private/GameFramework/GameplayCameraComponentBase.cpp
@@ -46,9 +46,6 @@ UGameplayCameraComponentBase::UGameplayCameraComponentBase(const FObjectInitiali
 
 	PrimaryComponentTick.bCanEverTick = true;
 	PrimaryComponentTick.TickGroup = TG_PostPhysics;
-
-	OutputCameraComponent = ObjectInit.CreateDefaultSubobject<UCineCameraComponent>(this, TEXT("OutputCameraComponent"), true);
-	OutputCameraComponent->SetupAttachment(this);
 }
 
 void UGameplayCameraComponentBase::BeginDestroy()
@@ -445,6 +442,14 @@ void UGameplayCameraComponentBase::OnRegister()
 
 	Super::OnRegister();
 
+	if (OutputCameraComponent == nullptr)
+	{
+		OutputCameraComponent = NewObject<UCineCameraComponent>(this, TEXT("OutputCameraComponent"), RF_Transient | RF_TextExportTransient);
+		OutputCameraComponent->SetupAttachment(this);
+		OutputCameraComponent->CreationMethod = CreationMethod;
+		OutputCameraComponent->RegisterComponentWithWorld(GetWorld());
+	}
+
 #if WITH_EDITOR
 
 	UWorld* World = GetWorld();
diff --git a/Engine/Plugins/Cameras/GameplayCameras/Source/GameplayCameras/Public/GameFramework/GameplayCameraComponentBase.h b/Engine/Plugins/Cameras/GameplayCameras/Source/GameplayCameras/Public/GameFramework/GameplayCameraComponentBase.h
--- a/Engine/Plugins/Cameras/GameplayCameras/Source/GameplayCameras/Public/GameFramework/GameplayCameraComponentBase.h
+++ b/Engine/Plugins/Cameras/GameplayCameras/Source/GameplayCameras/Public/GameFramework/GameplayCameraComponentBase.h
@@ -264,6 +264,10 @@ public:
 
 private:
 
+	/**
+	 * Camera component that receives the "output" of the hosted camera system when bRunStandaloneCameraSystem
+	 * is enabled.
+	 */
 	UPROPERTY(Transient)
 	TObjectPtr<UCineCameraComponent> OutputCameraComponent;
 
