Cameras: don't offer to create graph nodes when going "up-stream"

The ObjectTreeGraph model is mostly meant to go from "left to right", i.e. from an object to its properties to the objects connected to these properties. Dragging the "self" pin to the left only offered to create the exact same node type because that's the only node type compatible with itself generally (unless there are subclasses in exsistence). So it's mostly useless.

Also, it caused an undo bug since, after creating a node of the same type, the schema was trying to connect that to the dragged pin (the "self" pin on the first node). Which isn't possible. So it would fail, which would cancel the connection transaction, which would cancel the creation transaction. But the creation happened, so it made the creation of the (useless) node non-undoable. Now we just don't offer to create it anymore.

diff --git a/Engine/Plugins/Cameras/GameplayCameras/Source/GameplayCamerasEditor/Private/Editors/ObjectTreeGraphSchema.cpp b/Engine/Plugins/Cameras/GameplayCameras/Source/GameplayCamerasEditor/Private/Editors/ObjectTreeGraphSchema.cpp
--- a/Engine/Plugins/Cameras/GameplayCameras/Source/GameplayCamerasEditor/Private/Editors/ObjectTreeGraphSchema.cpp
+++ b/Engine/Plugins/Cameras/GameplayCameras/Source/GameplayCamerasEditor/Private/Editors/ObjectTreeGraphSchema.cpp
@@ -451,13 +451,18 @@ void UObjectTreeGraphSchema::GetGraphContextActions(FGraphContextMenuBuilder& Co
 		UObjectTreeGraphNode* OwningNode = Cast<UObjectTreeGraphNode>(DraggedPin->GetOwningNode());
 		if (OwningNode)
 		{
-			if (DraggedPin->PinType.PinCategory == UObjectTreeGraphSchema::PC_Self)
+			if (DraggedPin->PinType.PinCategory == UObjectTreeGraphSchema::PC_Property)
 			{
-				DraggedPinClass = OwningNode->GetObject()->GetClass();
+				DraggedPinClass = OwningNode->GetConnectedObjectClassForPin(DraggedPin);
 			}
-			else if (DraggedPin->PinType.PinCategory == UObjectTreeGraphSchema::PC_Property)
+			else if (DraggedPin->PinType.PinCategory == UObjectTreeGraphSchema::PC_Self)
 			{
-				DraggedPinClass = OwningNode->GetConnectedObjectClassForPin(DraggedPin);
+				// Dragged the self pin, so in theory we should offer to create nodes that have
+				// output pins (properties) that are compatible with our own class. That would
+				// require inspecting all properties on connectable classes, and hopefully
+				// caching it all somewhere... instead for now we just encourage the user to
+				// connect things from the left to the right and not go up-stream.
+				bShouldShowNewObjectActions = false;
 			}
 			else
 			{
