Cameras: improve how we apply parameter overrides from asset references

Remove the need to keep the cached copy of the parameter overrides to detect changes: instead, just compare the new override values to the values found in the variable and data tables.

diff --git a/Engine/Plugins/Cameras/GameplayCameras/Source/GameplayCameras/Private/Core/CameraAssetReference.cpp b/Engine/Plugins/Cameras/GameplayCameras/Source/GameplayCameras/Private/Core/CameraAssetReference.cpp
--- a/Engine/Plugins/Cameras/GameplayCameras/Source/GameplayCameras/Private/Core/CameraAssetReference.cpp
+++ b/Engine/Plugins/Cameras/GameplayCameras/Source/GameplayCameras/Private/Core/CameraAssetReference.cpp
@@ -23,17 +23,8 @@ void FCameraAssetReference::ApplyParameterOverrides(UE::Cameras::FCameraNodeEval
 	if (CameraAsset)
 	{
 		FCameraObjectInterfaceParameterOverrideHelper Helper(&OutResult.VariableTable, &OutResult.ContextDataTable);
-		Helper.ApplyParameterOverrides(CameraAsset, CameraAsset->GetParameterDefinitions(), Parameters, bDrivenOnly);
-	}
-}
-
-void FCameraAssetReference::ApplyParameterOverrides(const FInstancedPropertyBag& CachedParameters, UE::Cameras::FCameraNodeEvaluationResult& OutResult) const
-{
-	using namespace UE::Cameras;
-	if (CameraAsset)
-	{
-		FCameraObjectInterfaceParameterOverrideHelper Helper(&OutResult.VariableTable, &OutResult.ContextDataTable);
-		Helper.ApplyParameterOverrides(CameraAsset, CameraAsset->GetParameterDefinitions(), Parameters, CachedParameters);
+		Helper.bDrivenOnly = bDrivenOnly;
+		Helper.ApplyParameterOverrides(CameraAsset, CameraAsset->GetParameterDefinitions(), Parameters);
 	}
 }
 
diff --git a/Engine/Plugins/Cameras/GameplayCameras/Source/GameplayCameras/Private/Core/CameraRigAssetReference.cpp b/Engine/Plugins/Cameras/GameplayCameras/Source/GameplayCameras/Private/Core/CameraRigAssetReference.cpp
--- a/Engine/Plugins/Cameras/GameplayCameras/Source/GameplayCameras/Private/Core/CameraRigAssetReference.cpp
+++ b/Engine/Plugins/Cameras/GameplayCameras/Source/GameplayCameras/Private/Core/CameraRigAssetReference.cpp
@@ -42,17 +42,6 @@ void FCameraRigAssetReference::ApplyParameterOverrides(UE::Cameras::FCameraNodeE
 	ApplyParameterOverridesImpl(&OutResult.VariableTable, &OutResult.ContextDataTable, bDrivenOnly);
 }
 
-void FCameraRigAssetReference::ApplyParameterOverrides(const FInstancedPropertyBag& CachedParameters, UE::Cameras::FCameraNodeEvaluationResult& OutResult) const
-{
-	using namespace UE::Cameras;
-
-	if (CameraRig)
-	{
-		FCameraObjectInterfaceParameterOverrideHelper Helper(&OutResult.VariableTable, &OutResult.ContextDataTable);
-		Helper.ApplyParameterOverrides(CameraRig, CameraRig->GetParameterDefinitions(), Parameters, CachedParameters);
-	}
-}
-
 void FCameraRigAssetReference::ApplyParameterOverridesImpl(UE::Cameras::FCameraVariableTable* OutVariableTable, UE::Cameras::FCameraContextDataTable* OutContextDataTable, bool bDrivenOnly) const
 {
 	using namespace UE::Cameras;
@@ -60,7 +49,8 @@ void FCameraRigAssetReference::ApplyParameterOverridesImpl(UE::Cameras::FCameraV
 	if (CameraRig)
 	{
 		FCameraObjectInterfaceParameterOverrideHelper Helper(OutVariableTable, OutContextDataTable);
-		Helper.ApplyParameterOverrides(CameraRig, CameraRig->GetParameterDefinitions(), Parameters, bDrivenOnly);
+		Helper.bDrivenOnly = bDrivenOnly;
+		Helper.ApplyParameterOverrides(CameraRig, CameraRig->GetParameterDefinitions(), Parameters);
 	}
 }
 
diff --git a/Engine/Plugins/Cameras/GameplayCameras/Source/GameplayCameras/Private/Core/CameraShakeAssetReference.cpp b/Engine/Plugins/Cameras/GameplayCameras/Source/GameplayCameras/Private/Core/CameraShakeAssetReference.cpp
--- a/Engine/Plugins/Cameras/GameplayCameras/Source/GameplayCameras/Private/Core/CameraShakeAssetReference.cpp
+++ b/Engine/Plugins/Cameras/GameplayCameras/Source/GameplayCameras/Private/Core/CameraShakeAssetReference.cpp
@@ -44,7 +44,8 @@ void FCameraShakeAssetReference::ApplyParameterOverridesImpl(UE::Cameras::FCamer
 	if (CameraShake)
 	{
 		FCameraObjectInterfaceParameterOverrideHelper Helper(OutVariableTable, OutContextDataTable);
-		Helper.ApplyParameterOverrides(CameraShake, CameraShake->GetParameterDefinitions(), Parameters, bDrivenOnly);
+		Helper.bDrivenOnly = bDrivenOnly;
+		Helper.ApplyParameterOverrides(CameraShake, CameraShake->GetParameterDefinitions(), Parameters);
 	}
 }
 
diff --git a/Engine/Plugins/Cameras/GameplayCameras/Source/GameplayCameras/Private/GameFramework/GameplayCameraComponent.cpp b/Engine/Plugins/Cameras/GameplayCameras/Source/GameplayCameras/Private/GameFramework/GameplayCameraComponent.cpp
--- a/Engine/Plugins/Cameras/GameplayCameras/Source/GameplayCameras/Private/GameFramework/GameplayCameraComponent.cpp
+++ b/Engine/Plugins/Cameras/GameplayCameras/Source/GameplayCameras/Private/GameFramework/GameplayCameraComponent.cpp
@@ -4,6 +4,7 @@
 
 #include "Core/CameraAsset.h"
 #include "GameplayCamerasDelegates.h"
+#include "Helpers/CameraObjectInterfaceParameterOverrideHelper.h"
 #include "Misc/EngineVersionComparison.h"
 
 #include UE_INLINE_GENERATED_CPP_BY_NAME(GameplayCameraComponent)
@@ -72,25 +73,34 @@ void UGameplayCameraComponent::OnUpdateCameraEvaluationContext(bool bForceApplyP
 	// so we need to always re-apply values.
 	else
 	{
-		CameraReference.ApplyParameterOverrides(CachedParameterOverrides, InitialResult);
+		ApplyChangedParameterOverrides();
 	}
-
-	CachedParameterOverrides = CameraReference.GetParameters();
 #endif  // pre-5.7.0
 }
 
 void UGameplayCameraComponent::NotifyChangeCameraReference()
+{
+	// Sequencer animated some of our parameters... look for those whose value changed, compared to our
+	// cached parameter bag, and re-apply them to the evaluation context.
+	// TODO: This isn't a very efficient process, but it's unclear how to reconcile the reference's parameters struct
+	//		 with the evaluation context's variable/context-data tables.
+	ApplyChangedParameterOverrides();
+}
+
+void UGameplayCameraComponent::ApplyChangedParameterOverrides()
 {
 	using namespace UE::Cameras;
 
 	if (HasEvaluationContext())
 	{
-		// Sequencer animated some of our parameters... look for those whose value changed, compared to our
-		// cached parameter bag, and re-apply them to the evaluation context.
-		// TODO: This isn't a very efficient process, but it's unclear how to reconcile the reference's parameters struct
-		//		 with the evaluation context's variable/context-data tables.
+		const UCameraAsset* CameraAsset = CameraReference.GetCameraAsset();
 		FCameraNodeEvaluationResult& InitialResult = GetEvaluationContext()->GetInitialResult();
-		CameraReference.ApplyParameterOverrides(CachedParameterOverrides, InitialResult);
+
+		FCameraObjectInterfaceParameterOverrideHelper Helper(&InitialResult.VariableTable, &InitialResult.ContextDataTable);
+		Helper.bOverridenOnly = false;
+		GetChangedParameterOverrides(CameraReference.GetParameters(), CachedParameterOverrides, Helper.OnlyParameterGuids);
+		Helper.ApplyParameterOverrides(CameraAsset, CameraAsset->GetParameterDefinitions(), CameraReference.GetParameters());
+
 		CachedParameterOverrides = CameraReference.GetParameters();
 	}
 }
diff --git a/Engine/Plugins/Cameras/GameplayCameras/Source/GameplayCameras/Private/GameFramework/GameplayCameraComponentBase.cpp b/Engine/Plugins/Cameras/GameplayCameras/Source/GameplayCameras/Private/GameFramework/GameplayCameraComponentBase.cpp
--- a/Engine/Plugins/Cameras/GameplayCameras/Source/GameplayCameras/Private/GameFramework/GameplayCameraComponentBase.cpp
+++ b/Engine/Plugins/Cameras/GameplayCameras/Source/GameplayCameras/Private/GameFramework/GameplayCameraComponentBase.cpp
@@ -597,6 +597,51 @@ void UGameplayCameraComponentBase::UpdateEvaluationContext(bool bForceApplyParam
 	}
 }
 
+void UGameplayCameraComponentBase::GetChangedParameterOverrides(
+		const FInstancedPropertyBag& InParameterOverrides,
+		FInstancedPropertyBag& InOutCachedParameterOverrides,
+		TArray<FGuid>& OutChangedParameterGuids)
+{
+	const UPropertyBag* ParameterStruct = InParameterOverrides.GetPropertyBagStruct();
+	const UPropertyBag* CachedParameterStruct = InOutCachedParameterOverrides.GetPropertyBagStruct();
+	const uint8* ParametersContainer = InParameterOverrides.GetValue().GetMemory();
+	const uint8* CachedParameterContainer = InOutCachedParameterOverrides.GetValue().GetMemory();
+	if (!ParameterStruct || !ParametersContainer)
+	{
+		InOutCachedParameterOverrides.Reset();
+		return;
+	}
+	if (!CachedParameterStruct || CachedParameterStruct != ParameterStruct || !CachedParameterContainer)
+	{
+		for (const FPropertyBagPropertyDesc& PropertyDesc : ParameterStruct->GetPropertyDescs())
+		{
+			OutChangedParameterGuids.Add(PropertyDesc.ID);
+		}
+		InOutCachedParameterOverrides = InParameterOverrides;
+		return;
+	}
+
+	for (const FPropertyBagPropertyDesc& PropertyDesc : ParameterStruct->GetPropertyDescs())
+	{
+		if (!PropertyDesc.CachedProperty)
+		{
+			continue;
+		}
+
+		const void* ValuePtr = PropertyDesc.CachedProperty->ContainerPtrToValuePtr<void>(ParametersContainer);
+		const void* CachedValuePtr = PropertyDesc.CachedProperty->ContainerPtrToValuePtr<void>(CachedParameterContainer);
+		if (!ValuePtr || !CachedValuePtr)
+		{
+			continue;
+		}
+
+		if (!PropertyDesc.CachedProperty->Identical(ValuePtr, CachedValuePtr))
+		{
+			OutChangedParameterGuids.Add(PropertyDesc.ID);
+		}
+	}
+}
+
 void UGameplayCameraComponentBase::UpdateControlRotationIfNeeded()
 {
 	using namespace UE::Cameras;
diff --git a/Engine/Plugins/Cameras/GameplayCameras/Source/GameplayCameras/Private/GameFramework/GameplayCameraRigComponent.cpp b/Engine/Plugins/Cameras/GameplayCameras/Source/GameplayCameras/Private/GameFramework/GameplayCameraRigComponent.cpp
--- a/Engine/Plugins/Cameras/GameplayCameras/Source/GameplayCameras/Private/GameFramework/GameplayCameraRigComponent.cpp
+++ b/Engine/Plugins/Cameras/GameplayCameras/Source/GameplayCameras/Private/GameFramework/GameplayCameraRigComponent.cpp
@@ -9,6 +9,7 @@
 #include "Directors/SingleCameraDirector.h"
 #include "GameFramework/Actor.h"  // IWYU pragma: keep
 #include "GameplayCamerasDelegates.h"
+#include "Helpers/CameraObjectInterfaceParameterOverrideHelper.h"
 #include "Misc/EngineVersionComparison.h"
 
 #include UE_INLINE_GENERATED_CPP_BY_NAME(GameplayCameraRigComponent)
@@ -123,25 +124,34 @@ void UGameplayCameraRigComponent::OnUpdateCameraEvaluationContext(bool bForceApp
 	// so we need to always re-apply values.
 	else
 	{
-		CameraRigReference.ApplyParameterOverrides(CachedParameterOverrides, InitialResult);
+		ApplyChangedParameterOverrides();
 	}
-
-	CachedParameterOverrides = CameraRigReference.GetParameters();
 #endif  // pre-5.7.0
 }
 
 void UGameplayCameraRigComponent::NotifyChangeCameraRigReference()
+{
+	// Sequencer animated some of our parameters... look for those whose value changed, compared to our
+	// cached parameter bag, and re-apply them to the evaluation context.
+	// TODO: This isn't a very efficient process, but it's unclear how to reconcile the reference's parameters struct
+	//		 with the evaluation context's variable/context-data tables.
+	ApplyChangedParameterOverrides();
+}
+
+void UGameplayCameraRigComponent::ApplyChangedParameterOverrides()
 {
 	using namespace UE::Cameras;
 
 	if (HasEvaluationContext())
 	{
-		// Sequencer animated some of our parameters... look for those whose value changed, compared to our
-		// cached parameter bag, and re-apply them to the evaluation context.
-		// TODO: This isn't a very efficient process, but it's unclear how to reconcile the reference's parameters struct
-		//		 with the evaluation context's variable/context-data tables.
+		const UCameraRigAsset* CameraRigAsset = CameraRigReference.GetCameraRig();
 		FCameraNodeEvaluationResult& InitialResult = GetEvaluationContext()->GetInitialResult();
-		CameraRigReference.ApplyParameterOverrides(CachedParameterOverrides, InitialResult);
+
+		FCameraObjectInterfaceParameterOverrideHelper Helper(&InitialResult.VariableTable, &InitialResult.ContextDataTable);
+		Helper.bOverridenOnly = false;
+		GetChangedParameterOverrides(CameraRigReference.GetParameters(), CachedParameterOverrides, Helper.OnlyParameterGuids);
+		Helper.ApplyParameterOverrides(CameraRigAsset, CameraRigAsset->GetParameterDefinitions(), CameraRigReference.GetParameters());
+
 		CachedParameterOverrides = CameraRigReference.GetParameters();
 	}
 }
diff --git a/Engine/Plugins/Cameras/GameplayCameras/Source/GameplayCameras/Private/Helpers/CameraObjectInterfaceParameterOverrideHelper.cpp b/Engine/Plugins/Cameras/GameplayCameras/Source/GameplayCameras/Private/Helpers/CameraObjectInterfaceParameterOverrideHelper.cpp
--- a/Engine/Plugins/Cameras/GameplayCameras/Source/GameplayCameras/Private/Helpers/CameraObjectInterfaceParameterOverrideHelper.cpp
+++ b/Engine/Plugins/Cameras/GameplayCameras/Source/GameplayCameras/Private/Helpers/CameraObjectInterfaceParameterOverrideHelper.cpp
@@ -401,73 +401,31 @@ FCameraObjectInterfaceParameterOverrideHelper::FCameraObjectInterfaceParameterOv
 void FCameraObjectInterfaceParameterOverrideHelper::ApplyParameterOverrides(
 		const UObject* CameraObject,
 		TConstArrayView<FCameraObjectInterfaceParameterDefinition> ParameterDefinitions,
-		const FInstancedOverridablePropertyBag& ParameterOverrides,
-		bool bDrivenOnly)
+		const FInstancedOverridablePropertyBag& ParameterOverrides)
 {
 	using namespace Internal;
 
 	const UPropertyBag* ParameterOverridesStruct = ParameterOverrides.GetPropertyBagStruct();
-	if (!CameraObject || !ParameterOverridesStruct)
+	const uint8* ParameterOverridesMemory = ParameterOverrides.GetValue().GetMemory();
+	if (!CameraObject || !ParameterOverridesStruct || !ParameterOverridesMemory)
 	{
 		return;
 	}
 
 	for (const FCameraObjectInterfaceParameterDefinition& Definition : ParameterDefinitions)
 	{
-		if (!ContextDataTable && Definition.ParameterType == ECameraObjectInterfaceParameterType::Data)
+		if ((!VariableTable && Definition.ParameterType == ECameraObjectInterfaceParameterType::Blendable) ||
+				(!ContextDataTable && Definition.ParameterType == ECameraObjectInterfaceParameterType::Data))
 		{
 			continue;
 		}
 
-		if (!ParameterOverrides.IsPropertyOverriden(Definition.ParameterGuid))
+		if (bOverridenOnly && !ParameterOverrides.IsPropertyOverriden(Definition.ParameterGuid))
 		{
 			continue;
 		}
 
-		const FPropertyBagPropertyDesc* PropertyDesc = ParameterOverridesStruct->FindPropertyDescByID(Definition.ParameterGuid);
-		if (!ensure(PropertyDesc))
-		{
-			continue;
-		}
-
-		ApplyParameterOverride(CameraObject, Definition, ParameterOverrides, *PropertyDesc, bDrivenOnly);
-	}
-}
-
-void FCameraObjectInterfaceParameterOverrideHelper::ApplyParameterOverrides(
-		const UObject* CameraObject,
-		TConstArrayView<FCameraObjectInterfaceParameterDefinition> ParameterDefinitions,
-		const FInstancedOverridablePropertyBag& ParameterOverrides,
-		const FInstancedPropertyBag& CachedParameterOverrides)
-{
-	using namespace Internal;
-
-	const UPropertyBag* ParameterOverridesStruct = ParameterOverrides.GetPropertyBagStruct();
-	if (!CameraObject || !ParameterOverridesStruct)
-	{
-		return;
-	}
-
-	const uint8* ParameterOverridesMemory = ParameterOverrides.GetValue().GetMemory();
-	const uint8* CachedParameterOverridesMemory = CachedParameterOverrides.GetValue().GetMemory();
-	if (!ParameterOverridesMemory || !CachedParameterOverridesMemory)
-	{
-		return;
-	}
-
-	if (!ensure(ParameterOverridesStruct == CachedParameterOverrides.GetPropertyBagStruct()))
-	{
-		return;
-	}
-
-	for (const FCameraObjectInterfaceParameterDefinition& Definition : ParameterDefinitions)
-	{
-		if (!ContextDataTable && Definition.ParameterType == ECameraObjectInterfaceParameterType::Data)
-		{
-			continue;
-		}
-		
-		if (!ParameterOverrides.IsPropertyOverriden(Definition.ParameterGuid))
+		if (!OnlyParameterGuids.IsEmpty() && !OnlyParameterGuids.Contains(Definition.ParameterGuid))
 		{
 			continue;
 		}
@@ -477,13 +435,8 @@ void FCameraObjectInterfaceParameterOverrideHelper::ApplyParameterOverrides(
 		{
 			continue;
 		}
-
-		const void* RawValue = PropertyDesc->CachedProperty->ContainerPtrToValuePtr<void>(ParameterOverridesMemory);
-		const void* CachedRawValue = PropertyDesc->CachedProperty->ContainerPtrToValuePtr<void>(CachedParameterOverridesMemory);
-		if (!PropertyDesc->CachedProperty->Identical(RawValue, CachedRawValue))
-		{
-			ApplyParameterOverride(CameraObject, Definition, ParameterOverrides, *PropertyDesc, false);
-		}
+		
+		ApplyParameterOverride(CameraObject, Definition, ParameterOverrides, *PropertyDesc);
 	}
 }
 
@@ -491,8 +444,7 @@ void FCameraObjectInterfaceParameterOverrideHelper::ApplyParameterOverride(
 		const UObject* CameraObject,
 		const FCameraObjectInterfaceParameterDefinition& ParameterDefinition,
 		const FInstancedOverridablePropertyBag& PropertyBag,
-		const FPropertyBagPropertyDesc& PropertyBagPropertyDesc,
-		bool bDrivenOnly)
+		const FPropertyBagPropertyDesc& PropertyBagPropertyDesc)
 {
 	using namespace Internal;
 
diff --git a/Engine/Plugins/Cameras/GameplayCameras/Source/GameplayCameras/Private/Helpers/CameraObjectInterfaceParameterOverrideHelper.h b/Engine/Plugins/Cameras/GameplayCameras/Source/GameplayCameras/Private/Helpers/CameraObjectInterfaceParameterOverrideHelper.h
--- a/Engine/Plugins/Cameras/GameplayCameras/Source/GameplayCameras/Private/Helpers/CameraObjectInterfaceParameterOverrideHelper.h
+++ b/Engine/Plugins/Cameras/GameplayCameras/Source/GameplayCameras/Private/Helpers/CameraObjectInterfaceParameterOverrideHelper.h
@@ -52,15 +52,7 @@ public:
 	void ApplyParameterOverrides(
 			const UObject* CameraObject,
 			TConstArrayView<FCameraObjectInterfaceParameterDefinition> ParameterDefinitions,
-			const FInstancedOverridablePropertyBag& ParameterOverrides,
-			bool bDrivenOnly);
-
-	/** Sets overriden values of interface parameters in the given variable and context data tables. */
-	void ApplyParameterOverrides(
-			const UObject* CameraObject,
-			TConstArrayView<FCameraObjectInterfaceParameterDefinition> ParameterDefinitions,
-			const FInstancedOverridablePropertyBag& ParameterOverrides,
-			const FInstancedPropertyBag& CachedParameterOverrides);
+			const FInstancedOverridablePropertyBag& ParameterOverrides);
 
 private:
 
@@ -68,8 +60,16 @@ private:
 			const UObject* CameraObject,
 			const FCameraObjectInterfaceParameterDefinition& ParameterDefinition,
 			const FInstancedOverridablePropertyBag& PropertyBag,
-			const FPropertyBagPropertyDesc& PropertyBagPropertyDesc,
-			bool bDrivenOnly);
+			const FPropertyBagPropertyDesc& PropertyBagPropertyDesc);
+
+public:
+
+	/** Only parameters flagged as "overriden" should be applied. */
+	bool bOverridenOnly = true;
+	/** Only parameters that are dynamicaly driven should be applied. */
+	bool bDrivenOnly = false;
+	/** If not empty, the list of parameters we should only apply. */
+	TArray<FGuid> OnlyParameterGuids;
 
 private:
 
diff --git a/Engine/Plugins/Cameras/GameplayCameras/Source/GameplayCameras/Public/Core/CameraAssetReference.h b/Engine/Plugins/Cameras/GameplayCameras/Source/GameplayCameras/Public/Core/CameraAssetReference.h
--- a/Engine/Plugins/Cameras/GameplayCameras/Source/GameplayCameras/Public/Core/CameraAssetReference.h
+++ b/Engine/Plugins/Cameras/GameplayCameras/Source/GameplayCameras/Public/Core/CameraAssetReference.h
@@ -81,9 +81,6 @@ public:
 	/** Applies the parameter override values to the given evaluation result. */
 	GAMEPLAYCAMERAS_API void ApplyParameterOverrides(UE::Cameras::FCameraNodeEvaluationResult& OutResult, bool bDrivenOnly) const;
 
-	/** Applies the parameter override values to the given evaluation result. */
-	GAMEPLAYCAMERAS_API void ApplyParameterOverrides(const FInstancedPropertyBag& CachedParameters, UE::Cameras::FCameraNodeEvaluationResult& OutResult) const;
-
 public:
 
 	/** Returns whether the override parameters structure needs to be rebuilt. */
diff --git a/Engine/Plugins/Cameras/GameplayCameras/Source/GameplayCameras/Public/Core/CameraRigAssetReference.h b/Engine/Plugins/Cameras/GameplayCameras/Source/GameplayCameras/Public/Core/CameraRigAssetReference.h
--- a/Engine/Plugins/Cameras/GameplayCameras/Source/GameplayCameras/Public/Core/CameraRigAssetReference.h
+++ b/Engine/Plugins/Cameras/GameplayCameras/Source/GameplayCameras/Public/Core/CameraRigAssetReference.h
@@ -281,12 +281,12 @@ public:
 
 	/** Applies the parameter override values to the given variable table. */
 	void ApplyParameterOverrides(UE::Cameras::FCameraVariableTable& OutVariableTable, bool bDrivenOnly) const;
+
 	/** Applies the parameter override values to the given variable and context data tables. */
 	void ApplyParameterOverrides(UE::Cameras::FCameraVariableTable& OutVariableTable, UE::Cameras::FCameraContextDataTable& OutContextDataTable, bool bDrivenOnly) const;
+
 	/** Applies the parameter override values to the given evaluation result. */
 	void ApplyParameterOverrides(UE::Cameras::FCameraNodeEvaluationResult& OutResult, bool bDrivenOnly) const;
-	/** Applies the parameter override values to the given evaluation result. */
-	void ApplyParameterOverrides(const FInstancedPropertyBag& CachedParameters, UE::Cameras::FCameraNodeEvaluationResult& OutResult) const;
 
 private:
 
diff --git a/Engine/Plugins/Cameras/GameplayCameras/Source/GameplayCameras/Public/GameFramework/GameplayCameraComponent.h b/Engine/Plugins/Cameras/GameplayCameras/Source/GameplayCameras/Public/GameFramework/GameplayCameraComponent.h
--- a/Engine/Plugins/Cameras/GameplayCameras/Source/GameplayCameras/Public/GameFramework/GameplayCameraComponent.h
+++ b/Engine/Plugins/Cameras/GameplayCameras/Source/GameplayCameras/Public/GameFramework/GameplayCameraComponent.h
@@ -51,6 +51,8 @@ protected:
 
 private:
 
+	void ApplyChangedParameterOverrides();
+
 #if WITH_EDITOR
 	void OnCameraAssetBuilt(const UCameraAsset* InCameraAsset);
 #endif
diff --git a/Engine/Plugins/Cameras/GameplayCameras/Source/GameplayCameras/Public/GameFramework/GameplayCameraComponentBase.h b/Engine/Plugins/Cameras/GameplayCameras/Source/GameplayCameras/Public/GameFramework/GameplayCameraComponentBase.h
--- a/Engine/Plugins/Cameras/GameplayCameras/Source/GameplayCameras/Public/GameFramework/GameplayCameraComponentBase.h
+++ b/Engine/Plugins/Cameras/GameplayCameras/Source/GameplayCameras/Public/GameFramework/GameplayCameraComponentBase.h
@@ -197,6 +197,11 @@ protected:
 
 	void UpdateEvaluationContext(bool bForceApplyParameterOverrides);
 
+	void GetChangedParameterOverrides(
+			const FInstancedPropertyBag& InParameterOverrides,
+			FInstancedPropertyBag& InOutCachedParameterOverrides,
+			TArray<FGuid>& OutChangedParameterGuids);
+	
 	bool IsEditorWorld() const;
 	void UpdateControlRotationIfNeeded();
 
diff --git a/Engine/Plugins/Cameras/GameplayCameras/Source/GameplayCameras/Public/GameFramework/GameplayCameraRigComponent.h b/Engine/Plugins/Cameras/GameplayCameras/Source/GameplayCameras/Public/GameFramework/GameplayCameraRigComponent.h
--- a/Engine/Plugins/Cameras/GameplayCameras/Source/GameplayCameras/Public/GameFramework/GameplayCameraRigComponent.h
+++ b/Engine/Plugins/Cameras/GameplayCameras/Source/GameplayCameras/Public/GameFramework/GameplayCameraRigComponent.h
@@ -53,6 +53,8 @@ private:
 	void BuildGeneratedCamera();
 	void OnCameraRigAssetBuiltImpl();
 
+	void ApplyChangedParameterOverrides();
+
 #if WITH_EDITOR
 	void OnCameraRigAssetBuilt(const UCameraRigAsset* InCameraRigAsset);
 #endif
