Cameras: fix some problems with the screen-space framing nodes

- Don't stop reframing once the target is in the dead-zone, keep following it if it's still moving.
- Modulate the reframing based on the disengagement time to get a semi-random rest position inside the dead-zone.

diff --git a/Engine/Plugins/Cameras/GameplayCameras/Source/GameplayCameras/Private/Nodes/Framing/BaseFramingCameraNode.cpp b/Engine/Plugins/Cameras/GameplayCameras/Source/GameplayCameras/Private/Nodes/Framing/BaseFramingCameraNode.cpp
--- a/Engine/Plugins/Cameras/GameplayCameras/Source/GameplayCameras/Private/Nodes/Framing/BaseFramingCameraNode.cpp
+++ b/Engine/Plugins/Cameras/GameplayCameras/Source/GameplayCameras/Private/Nodes/Framing/BaseFramingCameraNode.cpp
@@ -2,8 +2,6 @@
 
 #include "Nodes/Framing/BaseFramingCameraNode.h"
 
-#include "CanvasItem.h"
-#include "CanvasTypes.h"
 #include "Core/CameraEvaluationContext.h"
 #include "Core/CameraParameterReader.h"
 #include "Debug/CameraDebugBlock.h"
@@ -15,6 +13,7 @@
 #include "HAL/IConsoleManager.h"
 #include "Math/CameraPoseMath.h"
 #include "Math/ColorList.h"
+#include "Math/Interpolation.h"
 #include "Misc/AssertionMacros.h"
 #include "Misc/CoreMiscDefines.h"
 #include "Nodes/Framing/CameraActorTargetInfo.h"
@@ -36,6 +35,18 @@ static FAutoConsoleVariableRef CVarFramingExtrapolationEpsilon(
 	GFramingExtrapolationEpsilon,
 	TEXT("(Default: 0.001) The epsilon to determine whether target movement extrapolation should be included."));
 
+float GFramingMinDampingFactor = 0.0001;
+static FAutoConsoleVariableRef CVarFramingMinDampingFactor(
+	TEXT("GameplayCameras.Framing.MinDampingFactor"),
+	GFramingExtrapolationEpsilon,
+	TEXT("(Default: 0.0001) The minimimum reframe damping factor possible once all factors have been taken into account."));
+
+float GFramingTargetRestEpsilon = 0.001;
+static FAutoConsoleVariableRef CVarFramingTargetRestEpsilon(
+	TEXT("GameplayCameras.Framing.TargetRestEpsilon"),
+	GFramingTargetRestEpsilon,
+	TEXT("(Default: 0.001) The screen-space speed under which we consider the target has stopped moving."));
+
 int32 GFramingNumTargetMovementSamples = 10;
 static FAutoConsoleVariableRef CVarFramingNumTargetMovementSamples(
 	TEXT("GameplayCameras.Framing.NumTargetMovementSamples"),
@@ -170,6 +181,9 @@ void FBaseFramingCameraNodeEvaluator::ComputeCurrentState(const FCameraNodeEvalu
 	// Get screen-space coordinates of the ideal framing point. These are in 0..1 UI space.
 	State.IdealTarget = Readers.IdealFramingLocation.Get(OutResult.VariableTable);
 
+	// Cache the aspect ratio.
+	State.AspectRatio = FCameraPoseMath::GetEffectiveAspectRatio(OutResult.CameraPose, Params.EvaluationContext);
+
 	// Update the damping factors and reengage/disengage times in case they are driven by a variable.
 	State.ReframeDampingFactor = Readers.ReframeDampingFactor.Get(OutResult.VariableTable);
 	State.LowReframeDampingFactor = Readers.LowReframeDampingFactor.Get(OutResult.VariableTable);
@@ -197,11 +211,15 @@ void FBaseFramingCameraNodeEvaluator::ComputeCurrentState(const FCameraNodeEvalu
 
 	// Process our targets and figure out the weighted average we should be aiming at.
 	FVector2d NewScreenTarget;
-	ComputeFinalTargetInfo(Params, TempPose, State.WorldTarget, NewScreenTarget, State.ScreenTargetBounds);
+	FVector3d OldWorldTarget = State.WorldTarget;
+	ComputeFinalTargetInfo(TempPose, State.WorldTarget, NewScreenTarget, State.ScreenTargetBounds);
 
 	// See if we need to extrapolate where the target will be in "anticipation time" seconds.
 	State.ScreenTarget = ComputeAnticipatedScreenTarget(Params.DeltaTime, State.ScreenTarget, NewScreenTarget);
 
+	// Compute how fast the target has moved over the last frame.
+	State.TargetScreenSpaceSpeed = ComputeScreenTargetSpeed(Params.DeltaTime, TempPose, OldWorldTarget, State.WorldTarget);
+
 	// Compute the effective dead-zone, which is the subset of the dead-zone that encompasses as much
 	// of the target's bound as possible.
 	State.EffectiveDeadZone = ComputeEffectiveDeadZone();
@@ -221,8 +239,8 @@ void FBaseFramingCameraNodeEvaluator::ComputeCurrentState(const FCameraNodeEvalu
 	else
 	{
 		// Make a line between the ideal target and the current target. Note how it intersects the 
-		// boundaries dead zone and the hard zone. We will interpolate the damping factors from the
-		// first intersection to the second intersection.
+		// boundaries of the dead zone and the hard zone. We will interpolate the damping factors from 
+		// the first intersection to the second intersection.
 		const FVector2d IdealToCurrent(State.ScreenTarget - State.IdealTarget);
 		const double IdealToCurrentDistance = IdealToCurrent.Length();
 		if (IdealToCurrentDistance > 0.0)
@@ -248,9 +266,9 @@ void FBaseFramingCameraNodeEvaluator::ComputeCurrentState(const FCameraNodeEvalu
 		}
 	}
 
-	const bool bWasReframing = State.bIsReframingTarget;
 	const bool bIsInSoftZone = State.SoftZone.Contains(State.ScreenTarget);
 	const bool bIsInDeadZone = State.EffectiveDeadZone.Contains(State.ScreenTarget);
+	const bool bTargetStillMoving = (State.TargetScreenSpaceSpeed > GFramingTargetRestEpsilon);
 	if (!bIsInSoftZone)
 	{
 		// Target is out of view or outside the soft zone -- it's therefore in the hard zone and we will
@@ -260,9 +278,10 @@ void FBaseFramingCameraNodeEvaluator::ComputeCurrentState(const FCameraNodeEvalu
 		State.ToggleEngageAlpha = 1.f;
 		State.bIsReframingTarget = true;
 	}
-	else if (!bIsInDeadZone)
+	else if (!bIsInDeadZone || bTargetStillMoving)
 	{
-		// Target is in the soft zone so we will gently reframe it towards the ideal framing.
+		// Target is in the soft zone, or in the dead-zone but still moving enough that we don't want to
+		// disengage the reframing. We will gently reframe it towards the ideal framing.
 		const bool bWasInDeadZone = (State.TargetFramingState == ETargetFramingState::InDeadZone);
 		State.TargetFramingState = ETargetFramingState::InSoftZone;
 
@@ -342,7 +361,7 @@ void FBaseFramingCameraNodeEvaluator::ComputeCurrentState(const FCameraNodeEvalu
 	}
 }
 
-bool FBaseFramingCameraNodeEvaluator::ComputeFinalTargetInfo(const FCameraNodeEvaluationParams& Params, const FCameraPose& CameraPose, FVector3d& OutWorldTarget, FVector2d& OutScreenTarget, FFramingZone& OutScreenBounds)
+bool FBaseFramingCameraNodeEvaluator::ComputeFinalTargetInfo(const FCameraPose& CameraPose, FVector3d& OutWorldTarget, FVector2d& OutScreenTarget, FFramingZone& OutScreenBounds)
 {
 	TConstArrayView<FCameraActorComputedTargetInfo> TargetInfos(WorldTargets.TargetInfos);
 
@@ -352,8 +371,6 @@ bool FBaseFramingCameraNodeEvaluator::ComputeFinalTargetInfo(const FCameraNodeEv
 	}
 
 	// Start with projecting all the targets, and their bounds, on screen.
-	const double AspectRatio = FCameraPoseMath::GetEffectiveAspectRatio(CameraPose, Params.EvaluationContext);
-
 	struct FComputedTargetScreenInfo
 	{
 		FVector3d WorldTarget;
@@ -375,11 +392,11 @@ bool FBaseFramingCameraNodeEvaluator::ComputeFinalTargetInfo(const FCameraNodeEv
 		FComputedTargetScreenInfo& TargetScreenInfo(TargetScreenInfos[Index]);
 
 		const FVector3d& WorldTarget = TargetInfo.Transform.GetLocation();
-		const TOptional<FVector2d> ScreenTarget = FCameraPoseMath::ProjectWorldToScreen(CameraPose, AspectRatio, WorldTarget, true);
+		const TOptional<FVector2d> ScreenTarget = FCameraPoseMath::ProjectWorldToScreen(CameraPose, State.AspectRatio, WorldTarget, true);
 
 		TargetScreenInfo.WorldTarget = WorldTarget;
 		TargetScreenInfo.ScreenTarget = ScreenTarget.Get(FVector2d(0.5, 0.5));
-		TargetScreenInfo.ScreenBounds = ComputeScreenTargetBounds(CameraPose, AspectRatio, TargetInfo.Transform, TargetInfo.LocalBounds);
+		TargetScreenInfo.ScreenBounds = ComputeScreenTargetBounds(CameraPose, State.AspectRatio, TargetInfo.Transform, TargetInfo.LocalBounds);
 		TargetScreenInfo.WorldTargetDistance = FVector3d::Distance(CameraPose.GetLocation(), WorldTarget);
 		TargetScreenInfo.NormalizedWeight = TargetInfo.NormalizedWeight;
 
@@ -412,7 +429,7 @@ bool FBaseFramingCameraNodeEvaluator::ComputeFinalTargetInfo(const FCameraNodeEv
 
 	// Unproject the final screen target, and use the weigted average distance to get, roughly, what
 	// world-space target we might be looking at.
-	const FVector3d FinalWorldTarget = FCameraPoseMath::UnprojectScreenToWorld(CameraPose, AspectRatio, FinalScreenTarget, FinalWorldTargetDistance);
+	const FVector3d FinalWorldTarget = FCameraPoseMath::UnprojectScreenToWorld(CameraPose, State.AspectRatio, FinalScreenTarget, FinalWorldTargetDistance);
 
 	OutScreenTarget = FinalScreenTarget;
 	OutScreenBounds = FinalScreenBounds;
@@ -421,6 +438,24 @@ bool FBaseFramingCameraNodeEvaluator::ComputeFinalTargetInfo(const FCameraNodeEv
 	return true;
 }
 
+float FBaseFramingCameraNodeEvaluator::ComputeScreenTargetSpeed(float DeltaTime, const FCameraPose& CameraPose, const FVector3d& OldWorldTarget, const FVector3d& NewWorldTarget)
+{
+	if (DeltaTime <= 0.f)
+	{
+		// Return last frame's speed.
+		return State.TargetScreenSpaceSpeed;
+	}
+
+	TOptional<FVector2d> OldScreenTarget = FCameraPoseMath::ProjectWorldToScreen(CameraPose, State.AspectRatio, OldWorldTarget, true);
+	TOptional<FVector2d> NewScreenTarget = FCameraPoseMath::ProjectWorldToScreen(CameraPose, State.AspectRatio, NewWorldTarget, true);
+	if (!OldScreenTarget.IsSet() || !NewScreenTarget.IsSet())
+	{
+		return State.TargetScreenSpaceSpeed;
+	}
+
+	return (FVector2d::Distance(OldScreenTarget.GetValue(), NewScreenTarget.GetValue()) / DeltaTime);
+}
+
 FVector2d FBaseFramingCameraNodeEvaluator::ComputeAnticipatedScreenTarget(float DeltaTime, const FVector2d& InPreviousAnticipatedScreenTarget, const FVector2d& InScreenTarget)
 {
 	if (State.TargetMovementAnticipationTime <= 0)
@@ -510,11 +545,9 @@ FFramingZone FBaseFramingCameraNodeEvaluator::ComputeScreenTargetBounds(const FC
 
 void FBaseFramingCameraNodeEvaluator::ComputeDesiredState(const FCameraNodeEvaluationParams& Params, const FCameraNodeEvaluationResult& OutResult)
 {
-	// If we  don't have any reframing to do, bail out.
-	FVector2d IdealToTarget(State.ScreenTarget - State.IdealTarget);
-	double DistanceToGo = IdealToTarget.Length();
-	const double ReframingSpeed = State.ReframeDamper.GetX0Derivative();
-	if (!State.bIsReframingTarget || (DistanceToGo <= GFramingIdealReachedEpsilon && ReframingSpeed <= GFramingIdealReachedEpsilon))
+	// If we don't have any reframing to do, bail out.
+	const double ReframingSpeed = FMath::Abs(State.ReframeDamper.GetX0Derivative());
+	if (!State.bIsReframingTarget && (ReframingSpeed <= GFramingIdealReachedEpsilon || State.ToggleEngageAlpha <= 0))
 	{
 		Desired.ScreenTarget = State.ScreenTarget;
 		Desired.FramingCorrection = FVector2d::ZeroVector;
@@ -532,6 +565,8 @@ void FBaseFramingCameraNodeEvaluator::ComputeDesiredState(const FCameraNodeEvalu
 		return;
 	}
 
+	FVector2d IdealToTarget(State.ScreenTarget - State.IdealTarget);
+	double DistanceToGo = IdealToTarget.Length();
 	if (State.TargetFramingState == ETargetFramingState::InHardZone)
 	{
 		// Bring the target immediately to edge of the soft zone, in the direction of the 
@@ -543,19 +578,22 @@ void FBaseFramingCameraNodeEvaluator::ComputeDesiredState(const FCameraNodeEvalu
 	}
 
 	// Figure out the damping factor for this frame. We might have interpolation between
-	// the low and normal (high) damping factors, and then we might have interpolation
-	// between that and 0 for disengaging or reengaging framing.
+	// the low and normal (high) damping factors if that option is enabled.
 	float EffectiveDampingFactor = State.ReframeDampingFactor;
 	if (State.LowReframeDampingFactor > 0)
 	{
 		EffectiveDampingFactor = FMath::Lerp(
 				State.LowReframeDampingFactor, State.ReframeDampingFactor, State.ReframeDampingFactorAlpha);
 	}
-	EffectiveDampingFactor = FMath::Lerp(0.1f, FMath::Max(0.1f, EffectiveDampingFactor), State.ToggleEngageAlpha);
+	EffectiveDampingFactor = FMath::Max(GFramingMinDampingFactor, EffectiveDampingFactor);
 	State.ReframeDamper.SetW0(EffectiveDampingFactor);
 
 	// Move the target towards the ideal framing using damping.
-	const double NewDistanceToGo = State.ReframeDamper.Update(DistanceToGo, Params.DeltaTime);
+	double NewDistanceToGo = State.ReframeDamper.Update(DistanceToGo, Params.DeltaTime);
+
+	// Slowly shorten the distance left to go if we're trying to disengage, or just starting
+	// to re-engage.
+	NewDistanceToGo = FMath::Lerp(DistanceToGo, NewDistanceToGo, SmoothStep(State.ToggleEngageAlpha));
 
 	// Compute where we want the target this frame.
 	const FVector2d InvReframeDir(IdealToTarget / DistanceToGo);
diff --git a/Engine/Plugins/Cameras/GameplayCameras/Source/GameplayCameras/Public/Nodes/Framing/BaseFramingCameraNode.h b/Engine/Plugins/Cameras/GameplayCameras/Source/GameplayCameras/Public/Nodes/Framing/BaseFramingCameraNode.h
--- a/Engine/Plugins/Cameras/GameplayCameras/Source/GameplayCameras/Public/Nodes/Framing/BaseFramingCameraNode.h
+++ b/Engine/Plugins/Cameras/GameplayCameras/Source/GameplayCameras/Public/Nodes/Framing/BaseFramingCameraNode.h
@@ -194,7 +194,8 @@ private:
 	void ComputeDesiredState(const FCameraNodeEvaluationParams& Params, const FCameraNodeEvaluationResult& OutResult);
 
 	bool AcquireTargetInfo(TSharedPtr<const FCameraEvaluationContext> EvaluationContext, const FCameraNodeEvaluationResult& InResult, TArray<FCameraActorComputedTargetInfo>& OutInfos);
-	bool ComputeFinalTargetInfo(const FCameraNodeEvaluationParams& Params, const FCameraPose& CameraPose, FVector3d& OutWorldTarget, FVector2d& OutScreenTarget, FFramingZone& OutScreenBounds);
+	bool ComputeFinalTargetInfo(const FCameraPose& CameraPose, FVector3d& OutWorldTarget, FVector2d& OutScreenTarget, FFramingZone& OutScreenBounds);
+	float ComputeScreenTargetSpeed(float DeltaTime, const FCameraPose& CameraPose, const FVector3d& OldWorldTarget, const FVector3d& NewWorldTarget);
 	FVector2d ComputeAnticipatedScreenTarget(float DeltaTime, const FVector2d& InPreviousAnticipatedScreenTarget, const FVector2d& InScreenTarget);
 	FFramingZone ComputeEffectiveDeadZone();
 
@@ -247,6 +248,8 @@ protected:
 	{
 		/** Screen-space position of the ideal framing position. */
 		FVector2d IdealTarget;
+		/** Current aspect ratio of the screen. */
+		float AspectRatio;
 		/** Current reframing damping factor. */
 		float ReframeDampingFactor;
 		/** Current low reframing damping factor. */
@@ -279,6 +282,8 @@ protected:
 
 		/** Current state of the tracked target. */
 		ETargetFramingState TargetFramingState;
+		/** Current screen-space speed of the tracked target. */
+		float TargetScreenSpaceSpeed = 0.f;
 		/** Whether we are actively trying to bring the target back to the ideal position. */
 		bool bIsReframingTarget = false;
 
