Cameras: boom arm node now supports additive rotation mode

diff --git a/Engine/Plugins/Cameras/GameplayCameras/Source/GameplayCameras/Private/Nodes/Common/BoomArmCameraNode.cpp b/Engine/Plugins/Cameras/GameplayCameras/Source/GameplayCameras/Private/Nodes/Common/BoomArmCameraNode.cpp
--- a/Engine/Plugins/Cameras/GameplayCameras/Source/GameplayCameras/Private/Nodes/Common/BoomArmCameraNode.cpp
+++ b/Engine/Plugins/Cameras/GameplayCameras/Source/GameplayCameras/Private/Nodes/Common/BoomArmCameraNode.cpp
@@ -76,7 +76,13 @@ FRotator FBoomArmCameraNodeEvaluator::ComputeBoomRotation(const FCameraNodeEvalu
 
 void FBoomArmCameraNodeEvaluator::OnRun(const FCameraNodeEvaluationParams& Params, FCameraNodeEvaluationResult& OutResult)
 {
-	const FRotator3d BoomRotation = ComputeBoomRotation(Params, OutResult);
+	FRotator3d BoomRotation = ComputeBoomRotation(Params, OutResult);
+
+	const UBoomArmCameraNode* BoomArmNode = GetCameraNodeAs<UBoomArmCameraNode>();
+	if (BoomArmNode->bAdditiveRotation)
+	{
+		BoomRotation = OutResult.CameraPose.GetRotation() + BoomRotation;
+	}
 
 	// Here we want to logically apply transform in this order:
 	//
diff --git a/Engine/Plugins/Cameras/GameplayCameras/Source/GameplayCameras/Public/Nodes/Common/BoomArmCameraNode.h b/Engine/Plugins/Cameras/GameplayCameras/Source/GameplayCameras/Public/Nodes/Common/BoomArmCameraNode.h
--- a/Engine/Plugins/Cameras/GameplayCameras/Source/GameplayCameras/Public/Nodes/Common/BoomArmCameraNode.h
+++ b/Engine/Plugins/Cameras/GameplayCameras/Source/GameplayCameras/Public/Nodes/Common/BoomArmCameraNode.h
@@ -110,6 +110,13 @@ public:
 	 */
 	UPROPERTY(meta=(ObjectTreeGraphPinDirection=Input))
 	TObjectPtr<UInput2DCameraNode> InputSlot;
+
+	/* 
+	* If true, the boom arm's rotation will add to the existing rotation instead of being absolute.
+	*/
+	UPROPERTY(EditAnywhere, Category=Common)
+	bool bAdditiveRotation = false;
+
 };
 
 #undef UE_API
