Cameras: protect panning framing node from introducing NaN rotations

diff --git a/Engine/Plugins/Cameras/GameplayCameras/Source/GameplayCameras/Private/Nodes/Framing/PanningFramingCameraNode.cpp b/Engine/Plugins/Cameras/GameplayCameras/Source/GameplayCameras/Private/Nodes/Framing/PanningFramingCameraNode.cpp
--- a/Engine/Plugins/Cameras/GameplayCameras/Source/GameplayCameras/Private/Nodes/Framing/PanningFramingCameraNode.cpp
+++ b/Engine/Plugins/Cameras/GameplayCameras/Source/GameplayCameras/Private/Nodes/Framing/PanningFramingCameraNode.cpp
@@ -79,7 +79,11 @@ void FPanningFramingCameraNodeEvaluator::OnRun(const FCameraNodeEvaluationParams
 			if (!InitialDesiredAim.IsNearlyZero())
 			{
 				const FVector3d InitialPoseAim = OutResult.CameraPose.GetAimDir();
-				PanningRotation = InitialDesiredAim.Rotation() - InitialPoseAim.Rotation();
+				const FRotator3d NewPanningRotation = InitialDesiredAim.Rotation() - InitialPoseAim.Rotation();
+				if (ensure(!NewPanningRotation.ContainsNaN()))
+				{
+					PanningRotation = NewPanningRotation;
+				}
 			}
 		}
 	}
@@ -104,11 +108,14 @@ void FPanningFramingCameraNodeEvaluator::OnRun(const FCameraNodeEvaluationParams
 
 		// We need to reverse pitch because rotators treat positive pitch as up, whereas we were 
 		// treating that as negative (due to the -1..1 normalized space of our math).
-		FRotator3d PanningCorrection(
+		const FRotator3d PanningCorrection(
 				(DesiredAngles.Y - TargetAngles.Y),	// Pitch
 				(TargetAngles.X - DesiredAngles.X),	// Yaw
 				0);									// Roll
-		PanningRotation += PanningCorrection;
+		if (ensure(!PanningCorrection.ContainsNaN()))
+		{
+			PanningRotation += PanningCorrection;
+		}
 
 #if UE_GAMEPLAY_CAMERAS_DEBUG
 		DebugPanningCorrection = PanningCorrection;
