Cameras: auto-rotate input node now supports going towards a zero rotation

diff --git a/Engine/Plugins/Cameras/GameplayCameras/Source/GameplayCameras/Private/Nodes/Input/AutoRotateInput2DCameraNode.cpp b/Engine/Plugins/Cameras/GameplayCameras/Source/GameplayCameras/Private/Nodes/Input/AutoRotateInput2DCameraNode.cpp
--- a/Engine/Plugins/Cameras/GameplayCameras/Source/GameplayCameras/Private/Nodes/Input/AutoRotateInput2DCameraNode.cpp
+++ b/Engine/Plugins/Cameras/GameplayCameras/Source/GameplayCameras/Private/Nodes/Input/AutoRotateInput2DCameraNode.cpp
@@ -214,6 +214,10 @@ void FAutoRotateInput2DCameraNodeEvaluator::OnRun(const FCameraNodeEvaluationPar
 				}
 				bHasAutoRotateDir = true;
 				break;
+			case ECameraAutoRotateDirection::Zero:
+				AutoRotateDir = FVector3d(1, 0, 0);
+				bHasAutoRotateDir = true;
+				break;
 		}
 	}
 	if (!bHasAutoRotateDir)
diff --git a/Engine/Plugins/Cameras/GameplayCameras/Source/GameplayCameras/Public/Nodes/Input/AutoRotateInput2DCameraNode.h b/Engine/Plugins/Cameras/GameplayCameras/Source/GameplayCameras/Public/Nodes/Input/AutoRotateInput2DCameraNode.h
--- a/Engine/Plugins/Cameras/GameplayCameras/Source/GameplayCameras/Public/Nodes/Input/AutoRotateInput2DCameraNode.h
+++ b/Engine/Plugins/Cameras/GameplayCameras/Source/GameplayCameras/Public/Nodes/Input/AutoRotateInput2DCameraNode.h
@@ -26,7 +26,11 @@ enum class ECameraAutoRotateDirection
 	 * Re-align towards the evaluation context's movement direction if there is movement,
 	 * or towards its facing otherwise.
 	 */
-	MovementOrFacing
+	MovementOrFacing,
+	/**
+	 * Re-align towards the zero direction.
+	 */
+	Zero,
 };
 
 /**
