Cameras: add parameter GUIDs as tool-tips in the interface parameters panel

diff --git a/Engine/Plugins/Cameras/GameplayCameras/Source/GameplayCamerasEditor/Private/Toolkits/CameraObjectInterfaceParametersToolkit.cpp b/Engine/Plugins/Cameras/GameplayCameras/Source/GameplayCamerasEditor/Private/Toolkits/CameraObjectInterfaceParametersToolkit.cpp
--- a/Engine/Plugins/Cameras/GameplayCameras/Source/GameplayCamerasEditor/Private/Toolkits/CameraObjectInterfaceParametersToolkit.cpp
+++ b/Engine/Plugins/Cameras/GameplayCameras/Source/GameplayCamerasEditor/Private/Toolkits/CameraObjectInterfaceParametersToolkit.cpp
@@ -77,7 +77,8 @@ protected:
 			return SAssignNew(NameTextBlock, SInlineEditableTextBlock)
 				.IsSelected(this, &SCameraObjectInterfaceParameterTableRowBase<ParameterType>::IsSelected)
 				.Text_Lambda([this]() { return FText::FromString(Item->InterfaceParameterName); })
-				.OnTextCommitted(this, &SCameraObjectInterfaceParameterTableRowBase<ParameterType>::OnParameterNameTextCommitted);
+				.OnTextCommitted(this, &SCameraObjectInterfaceParameterTableRowBase<ParameterType>::OnParameterNameTextCommitted)
+				.ToolTipText(this, &SCameraObjectInterfaceParameterTableRowBase<ParameterType>::GetParameterToolTip);
 		}
 		else if (InColumnName == ParameterMessageColumn)
 		{
@@ -110,6 +111,12 @@ protected:
 		Item->InterfaceParameterName = Text.ToString();
 	}
 
+	FText GetParameterToolTip() const
+	{
+		FGuid Guid = Item->GetGuid();
+		return FText::FromString(Guid.ToString(EGuidFormats::DigitsWithHyphensInParentheses));
+	}
+
 	FText GetParameterMessageText() const
 	{
 		if (!Item->Target || Item->TargetPropertyName.IsNone())
