Cameras: make sure components' parameter overrides are up-to-date with the referenced asset

Update parameter overrides struct when starting evaluation, and on cook.

diff --git a/Engine/Plugins/Cameras/GameplayCameras/Source/GameplayCameras/Private/GameFramework/GameplayCameraComponent.cpp b/Engine/Plugins/Cameras/GameplayCameras/Source/GameplayCameras/Private/GameFramework/GameplayCameraComponent.cpp
--- a/Engine/Plugins/Cameras/GameplayCameras/Source/GameplayCameras/Private/GameFramework/GameplayCameraComponent.cpp
+++ b/Engine/Plugins/Cameras/GameplayCameras/Source/GameplayCameras/Private/GameFramework/GameplayCameraComponent.cpp
@@ -5,7 +5,9 @@
 #include "Core/CameraAsset.h"
 #include "GameplayCamerasDelegates.h"
 #include "Helpers/CameraObjectInterfaceParameterOverrideHelper.h"
+#include "Interfaces/ITargetPlatform.h"
 #include "Misc/EngineVersionComparison.h"
+#include "UObject/ObjectSaveContext.h"
 
 #include UE_INLINE_GENERATED_CPP_BY_NAME(GameplayCameraComponent)
 
@@ -27,6 +29,23 @@ void UGameplayCameraComponent::PostLoad()
 	}
 }
 
+void UGameplayCameraComponent::PreSave(FObjectPreSaveContext ObjectSaveContext)
+{
+#if WITH_EDITOR
+	if (!HasAnyFlags(RF_ClassDefaultObject | RF_ArchetypeObject))
+	{
+		// Make sure parameters are up to date when cooking.
+		const ITargetPlatform* TargetPlatform = ObjectSaveContext.GetTargetPlatform();
+		if (TargetPlatform && TargetPlatform->RequiresCookedData())
+		{
+			CameraReference.RebuildParametersIfNeeded();
+		}
+	}
+#endif  // WITH_EDITOR
+
+	Super::PreSave(ObjectSaveContext);
+}
+
 void UGameplayCameraComponent::OnRegister()
 {
 	using namespace UE::Cameras;
@@ -53,12 +72,16 @@ void UGameplayCameraComponent::OnUnregister()
 	Super::OnUnregister();
 }
 
-UCameraAsset* UGameplayCameraComponent::GetCameraAsset()
+UCameraAsset* UGameplayCameraComponent::OnCreateEvaluationContext()
 {
+#if WITH_EDITOR
+	CameraReference.RebuildParametersIfNeeded();
+#endif  // WITH_EDITOR
+
 	return CameraReference.GetCameraAsset();
 }
 
-void UGameplayCameraComponent::OnUpdateCameraEvaluationContext(bool bForceApplyParameterOverrides)
+void UGameplayCameraComponent::OnUpdateEvaluationContext(bool bForceApplyParameterOverrides)
 {
 	using namespace UE::Cameras;
 
@@ -118,12 +141,16 @@ void UGameplayCameraComponent::OnCameraAssetBuilt(const UCameraAsset* InCameraAs
 
 	// If our camera asset was just built, it may have some new parameters. We need to rebuild
 	// our variable table and context data table, and re-apply overrides.
-	CameraReference.RebuildParametersIfNeeded();
-	if (HasEvaluationContext())
+	if (CameraReference.NeedsRebuildParameters())
 	{
-		const FCameraAssetAllocationInfo& AllocationInfo = InCameraAsset->GetAllocationInfo();
-		ReinitializeEvaluationContext(AllocationInfo.VariableTableInfo, AllocationInfo.ContextDataTableInfo);
-		UpdateEvaluationContext(true);
+		Modify();
+		CameraReference.RebuildParameters();
+		if (HasEvaluationContext())
+		{
+			const FCameraAssetAllocationInfo& AllocationInfo = InCameraAsset->GetAllocationInfo();
+			ReinitializeEvaluationContext(AllocationInfo.VariableTableInfo, AllocationInfo.ContextDataTableInfo);
+			UpdateEvaluationContext(true);
+		}
 	}
 }
 
diff --git a/Engine/Plugins/Cameras/GameplayCameras/Source/GameplayCameras/Private/GameFramework/GameplayCameraComponentBase.cpp b/Engine/Plugins/Cameras/GameplayCameras/Source/GameplayCameras/Private/GameFramework/GameplayCameraComponentBase.cpp
--- a/Engine/Plugins/Cameras/GameplayCameras/Source/GameplayCameras/Private/GameFramework/GameplayCameraComponentBase.cpp
+++ b/Engine/Plugins/Cameras/GameplayCameras/Source/GameplayCameras/Private/GameFramework/GameplayCameraComponentBase.cpp
@@ -225,7 +225,7 @@ void UGameplayCameraComponentBase::EnsureEvaluationContext(APlayerController* Pl
 
 	if (!EvaluationContext.IsValid())
 	{
-		UCameraAsset* CameraAsset = GetCameraAsset();
+		UCameraAsset* CameraAsset = OnCreateEvaluationContext();
 
 		// If we have no camera asset specified, make a placeholder one and log a warning.
 		if (!CameraAsset)
@@ -587,7 +587,7 @@ void UGameplayCameraComponentBase::UpdateEvaluationContext(bool bForceApplyParam
 			bIsCameraCutNextFrame = false;
 		}
 
-		OnUpdateCameraEvaluationContext(bForceApplyParameterOverrides);
+		OnUpdateEvaluationContext(bForceApplyParameterOverrides);
 
 #if WITH_EDITOR
 
diff --git a/Engine/Plugins/Cameras/GameplayCameras/Source/GameplayCameras/Private/GameFramework/GameplayCameraRigComponent.cpp b/Engine/Plugins/Cameras/GameplayCameras/Source/GameplayCameras/Private/GameFramework/GameplayCameraRigComponent.cpp
--- a/Engine/Plugins/Cameras/GameplayCameras/Source/GameplayCameras/Private/GameFramework/GameplayCameraRigComponent.cpp
+++ b/Engine/Plugins/Cameras/GameplayCameras/Source/GameplayCameras/Private/GameFramework/GameplayCameraRigComponent.cpp
@@ -10,7 +10,9 @@
 #include "GameFramework/Actor.h"  // IWYU pragma: keep
 #include "GameplayCamerasDelegates.h"
 #include "Helpers/CameraObjectInterfaceParameterOverrideHelper.h"
+#include "Interfaces/ITargetPlatform.h"
 #include "Misc/EngineVersionComparison.h"
+#include "UObject/ObjectSaveContext.h"
 
 #include UE_INLINE_GENERATED_CPP_BY_NAME(GameplayCameraRigComponent)
 
@@ -19,6 +21,23 @@ UGameplayCameraRigComponent::UGameplayCameraRigComponent(const FObjectInitialize
 {
 }
 
+void UGameplayCameraRigComponent::PreSave(FObjectPreSaveContext ObjectSaveContext)
+{
+#if WITH_EDITOR
+	if (!HasAnyFlags(RF_ClassDefaultObject | RF_ArchetypeObject))
+	{
+		// Make sure parameters are up to date when cooking.
+		const ITargetPlatform* TargetPlatform = ObjectSaveContext.GetTargetPlatform();
+		if (TargetPlatform && TargetPlatform->RequiresCookedData())
+		{
+			CameraRigReference.RebuildParametersIfNeeded();
+		}
+	}
+#endif  // WITH_EDITOR
+
+	Super::PreSave(ObjectSaveContext);
+}
+
 void UGameplayCameraRigComponent::OnRegister()
 {
 	using namespace UE::Cameras;
@@ -45,21 +64,16 @@ void UGameplayCameraRigComponent::OnUnregister()
 	Super::OnUnregister();
 }
 
-UCameraAsset* UGameplayCameraRigComponent::GetCameraAsset()
+UCameraAsset* UGameplayCameraRigComponent::OnCreateEvaluationContext()
 {
-	if (!GeneratedCameraAsset)
+#if WITH_EDITOR
+	CameraRigReference.RebuildParametersIfNeeded();
+#endif  // WITH_EDITOR
+
+	if (!GeneratedCameraAsset && CameraRigReference.GetCameraRig())
 	{
 		USingleCameraDirector* SingleDirector = NewObject<USingleCameraDirector>(this, TEXT("GeneratedCameraDirector"), RF_Transient);
 		SingleDirector->CameraRig = CameraRigReference.GetCameraRig();
-		if (!SingleDirector->CameraRig)
-		{
-			UE_LOG(LogCameraSystem, Warning, 
-					TEXT("No camera rig specified on Gameplay Camera component '%s.%s', using a placeholder one."),
-					*GetNameSafe(GetOwner()), *GetNameSafe(this));
-
-			UCameraRigAsset* PlaceholderCameraRig = NewObject<UCameraRigAsset>();
-			SingleDirector->CameraRig = PlaceholderCameraRig;
-		}
 
 		GeneratedCameraAsset = NewObject<UCameraAsset>(this, TEXT("GeneratedCameraAsset"), RF_Transient);
 		GeneratedCameraAsset->SetCameraDirector(SingleDirector);
@@ -96,20 +110,24 @@ void UGameplayCameraRigComponent::OnCameraRigAssetBuiltImpl()
 {
 	using namespace UE::Cameras;
 
-	CameraRigReference.RebuildParametersIfNeeded();
-	if (HasEvaluationContext())
+	if (CameraRigReference.NeedsRebuildParameters())
 	{
+		Modify();
+		CameraRigReference.RebuildParameters();
+		if (HasEvaluationContext())
+		{
 #if WITH_EDITOR
-		const UCameraRigAsset* CameraRigAsset = CameraRigReference.GetCameraRig();
-		const FCameraObjectAllocationInfo& AllocationInfo = CameraRigAsset->AllocationInfo;
-		ReinitializeEvaluationContext(AllocationInfo.VariableTableInfo, AllocationInfo.ContextDataTableInfo);
+			const UCameraRigAsset* CameraRigAsset = CameraRigReference.GetCameraRig();
+			const FCameraObjectAllocationInfo& AllocationInfo = CameraRigAsset->AllocationInfo;
+			ReinitializeEvaluationContext(AllocationInfo.VariableTableInfo, AllocationInfo.ContextDataTableInfo);
 #endif  // WITH_EDITOR
 
-		UpdateEvaluationContext(true);
+			UpdateEvaluationContext(true);
+		}
 	}
 }
 
-void UGameplayCameraRigComponent::OnUpdateCameraEvaluationContext(bool bForceApplyParameterOverrides)
+void UGameplayCameraRigComponent::OnUpdateEvaluationContext(bool bForceApplyParameterOverrides)
 {
 	using namespace UE::Cameras;
 
diff --git a/Engine/Plugins/Cameras/GameplayCameras/Source/GameplayCameras/Public/GameFramework/GameplayCameraComponent.h b/Engine/Plugins/Cameras/GameplayCameras/Source/GameplayCameras/Public/GameFramework/GameplayCameraComponent.h
--- a/Engine/Plugins/Cameras/GameplayCameras/Source/GameplayCameras/Public/GameFramework/GameplayCameraComponent.h
+++ b/Engine/Plugins/Cameras/GameplayCameras/Source/GameplayCameras/Public/GameFramework/GameplayCameraComponent.h
@@ -30,8 +30,9 @@ public:
 
 	// UObject interface.
 	virtual void PostLoad() override;
+	virtual void PreSave(FObjectPreSaveContext ObjectSaveContext) override;
 #if WITH_EDITOR
-	virtual void PostEditChangeProperty( struct FPropertyChangedEvent& PropertyChangedEvent) override;
+	virtual void PostEditChangeProperty(FPropertyChangedEvent& PropertyChangedEvent) override;
 #endif
 
 public:
@@ -46,8 +47,8 @@ protected:
 	void NotifyChangeCameraReference();
 
 	// UGameplayCameraComponentBase interface.
-	virtual UCameraAsset* GetCameraAsset() override;
-	virtual void OnUpdateCameraEvaluationContext(bool bForceApplyParameterOverrides) override;
+	virtual UCameraAsset* OnCreateEvaluationContext() override;
+	virtual void OnUpdateEvaluationContext(bool bForceApplyParameterOverrides) override;
 
 private:
 
diff --git a/Engine/Plugins/Cameras/GameplayCameras/Source/GameplayCameras/Public/GameFramework/GameplayCameraComponentBase.h b/Engine/Plugins/Cameras/GameplayCameras/Source/GameplayCameras/Public/GameFramework/GameplayCameraComponentBase.h
--- a/Engine/Plugins/Cameras/GameplayCameras/Source/GameplayCameras/Public/GameFramework/GameplayCameraComponentBase.h
+++ b/Engine/Plugins/Cameras/GameplayCameras/Source/GameplayCameras/Public/GameFramework/GameplayCameraComponentBase.h
@@ -192,8 +192,8 @@ public:
 protected:
 
 	// UGameplayCameraComponentBase interface.
-	virtual UCameraAsset* GetCameraAsset() PURE_VIRTUAL(UGameplayCameraComponentBase::GetCameraAsset, return nullptr;)
-	virtual void OnUpdateCameraEvaluationContext(bool bForceApplyParameterOverrides) {}
+	virtual UCameraAsset* OnCreateEvaluationContext() { return nullptr; }
+	virtual void OnUpdateEvaluationContext(bool bForceApplyParameterOverrides) {}
 
 	void UpdateEvaluationContext(bool bForceApplyParameterOverrides);
 
diff --git a/Engine/Plugins/Cameras/GameplayCameras/Source/GameplayCameras/Public/GameFramework/GameplayCameraRigComponent.h b/Engine/Plugins/Cameras/GameplayCameras/Source/GameplayCameras/Public/GameFramework/GameplayCameraRigComponent.h
--- a/Engine/Plugins/Cameras/GameplayCameras/Source/GameplayCameras/Public/GameFramework/GameplayCameraRigComponent.h
+++ b/Engine/Plugins/Cameras/GameplayCameras/Source/GameplayCameras/Public/GameFramework/GameplayCameraRigComponent.h
@@ -29,8 +29,9 @@ public:
 	virtual void OnUnregister() override;
 
 	// UObject interface.
+	virtual void PreSave(FObjectPreSaveContext ObjectSaveContext) override;
 #if WITH_EDITOR
-	virtual void PostEditChangeProperty( struct FPropertyChangedEvent& PropertyChangedEvent) override;
+	virtual void PostEditChangeProperty(FPropertyChangedEvent& PropertyChangedEvent) override;
 #endif
 
 public:
@@ -45,8 +46,8 @@ protected:
 	void NotifyChangeCameraRigReference();
 
 	// UGameplayCameraComponentBase interface.
-	virtual UCameraAsset* GetCameraAsset() override;
-	virtual void OnUpdateCameraEvaluationContext(bool bForceApplyParameterOverrides) override;
+	virtual UCameraAsset* OnCreateEvaluationContext() override;
+	virtual void OnUpdateEvaluationContext(bool bForceApplyParameterOverrides) override;
 
 private:
 
