Cameras: occlusion material node now supports skinned meshes and attached components

diff --git a/Engine/Plugins/Cameras/GameplayCameras/Source/GameplayCameras/Private/Nodes/Collision/OcclusionMaterialCameraNode.cpp b/Engine/Plugins/Cameras/GameplayCameras/Source/GameplayCameras/Private/Nodes/Collision/OcclusionMaterialCameraNode.cpp
--- a/Engine/Plugins/Cameras/GameplayCameras/Source/GameplayCameras/Private/Nodes/Collision/OcclusionMaterialCameraNode.cpp
+++ b/Engine/Plugins/Cameras/GameplayCameras/Source/GameplayCameras/Private/Nodes/Collision/OcclusionMaterialCameraNode.cpp
@@ -3,7 +3,7 @@
 #include "Nodes/Collision/OcclusionMaterialCameraNode.h"
 
 #include "CollisionQueryParams.h"
-#include "Components/StaticMeshComponent.h"
+#include "Components/PrimitiveComponent.h"
 #include "Containers/Map.h"
 #include "Containers/Set.h"
 #include "Core/CameraEvaluationContext.h"
@@ -47,10 +47,12 @@ private:
 	void RunOcclusionTrace(UWorld* World, const FCameraNodeEvaluationParams& Params, FCameraNodeEvaluationResult& OutResult);
 	void HandleOcclusionTraceResult(UWorld* World);
 
-	void ApplyOcclusionMaterial(TSet<UStaticMeshComponent*> MeshComponents);
-	void RemoveOcclusionMaterial(TSet<UStaticMeshComponent*> MeshComponents);
+	void FindAllPrimitiveComponents(UPrimitiveComponent* InPrimitiveComponent, TSet<UPrimitiveComponent*>& OutPrimitiveComponents);
 
-	void ResolveWeakMeshComponents(TSet<TWeakObjectPtr<UStaticMeshComponent>> WeakMeshComponents, TSet<UStaticMeshComponent*>& OutMeshComponents);
+	void ApplyOcclusionMaterial(TSet<UPrimitiveComponent*> PrimitiveComponents);
+	void RemoveOcclusionMaterial(TSet<UPrimitiveComponent*> PrimitiveComponents);
+
+	void ResolveWeakPrimitiveComponents(TSet<TWeakObjectPtr<UPrimitiveComponent>> WeakPrimitiveComponents, TSet<UPrimitiveComponent*>& OutPrimitiveComponents);
 
 private:
 
@@ -58,8 +60,8 @@ private:
 	TCameraParameterReader<FVector3d> OcclusionTargetOffsetReader;
 
 	FTraceHandle OcclusionTraceHandle;
-	TSet<TWeakObjectPtr<UStaticMeshComponent>> CurrentlyOccludedMeshComponents;
-	TMap<TWeakObjectPtr<UStaticMeshComponent>, FOcclusionMaterialOverrideInfo> AppliedMaterialOverrides;
+	TSet<TWeakObjectPtr<UPrimitiveComponent>> CurrentlyOccludedPrimitiveComponents;
+	TMap<TWeakObjectPtr<UPrimitiveComponent>, FOcclusionMaterialOverrideInfo> AppliedMaterialOverrides;
 };
 
 UE_DEFINE_CAMERA_NODE_EVALUATOR(FOcclusionMaterialCameraNodeEvaluator)
@@ -67,9 +69,9 @@ UE_DEFINE_CAMERA_NODE_EVALUATOR(FOcclusionMaterialCameraNodeEvaluator)
 FOcclusionMaterialCameraNodeEvaluator::~FOcclusionMaterialCameraNodeEvaluator()
 {
 	// Make sure any occluded meshes are released when our camera rig is deactivated it.
-	TSet<UStaticMeshComponent*> MeshComponents;
-	ResolveWeakMeshComponents(CurrentlyOccludedMeshComponents, MeshComponents);
-	RemoveOcclusionMaterial(MeshComponents);
+	TSet<UPrimitiveComponent*> PrimitiveComponents;
+	ResolveWeakPrimitiveComponents(CurrentlyOccludedPrimitiveComponents, PrimitiveComponents);
+	RemoveOcclusionMaterial(PrimitiveComponents);
 }
 
 void FOcclusionMaterialCameraNodeEvaluator::OnInitialize(const FCameraNodeEvaluatorInitializeParams& Params, FCameraNodeEvaluationResult& OutResult)
@@ -188,36 +190,45 @@ void FOcclusionMaterialCameraNodeEvaluator::HandleOcclusionTraceResult(UWorld* W
 
 	// Get the list of meshes collected by the occlusion trace, and figure out which ones are
 	// new, and which ones got out of the way.
-	TSet<UStaticMeshComponent*> MeshComponents;
+	TSet<UPrimitiveComponent*> PrimitiveComponents;
 	for (const FHitResult& Hit : TraceDatum.OutHits)
 	{
-		UStaticMeshComponent* MeshComponent = Cast<UStaticMeshComponent>(Hit.GetComponent());
-		if (MeshComponent)
-		{
-			MeshComponents.Add(MeshComponent);
-		}
+		FindAllPrimitiveComponents(Hit.GetComponent(), PrimitiveComponents);
 	}
 
-	TSet<UStaticMeshComponent*> CurrentMeshComponents;
-	ResolveWeakMeshComponents(CurrentlyOccludedMeshComponents, CurrentMeshComponents);
+	TSet<UPrimitiveComponent*> CurrentPrimitiveComponents;
+	ResolveWeakPrimitiveComponents(CurrentlyOccludedPrimitiveComponents, CurrentPrimitiveComponents);
 
-	TSet<UStaticMeshComponent*> NewMeshComponents = MeshComponents.Difference(CurrentMeshComponents);
-	TSet<UStaticMeshComponent*> OldMeshComponents = CurrentMeshComponents.Difference(MeshComponents);
+	TSet<UPrimitiveComponent*> NewPrimitiveComponents = PrimitiveComponents.Difference(CurrentPrimitiveComponents);
+	TSet<UPrimitiveComponent*> OldPrimitiveComponents = CurrentPrimitiveComponents.Difference(PrimitiveComponents);
 
-	CurrentlyOccludedMeshComponents.Reset();
-	for (UStaticMeshComponent* MeshComponent : MeshComponents)
+	CurrentlyOccludedPrimitiveComponents.Reset();
+	for (UPrimitiveComponent* PrimitiveComponent : PrimitiveComponents)
 	{
-		CurrentlyOccludedMeshComponents.Add(MeshComponent);
+		CurrentlyOccludedPrimitiveComponents.Add(PrimitiveComponent);
 	}
 
 	// Apply occlusion material changes to new/old components.
-	ApplyOcclusionMaterial(NewMeshComponents);
-	RemoveOcclusionMaterial(OldMeshComponents);
+	ApplyOcclusionMaterial(NewPrimitiveComponents);
+	RemoveOcclusionMaterial(OldPrimitiveComponents);
 
 	OcclusionTraceHandle.Invalidate();
 }
 
-void FOcclusionMaterialCameraNodeEvaluator::ApplyOcclusionMaterial(TSet<UStaticMeshComponent*> MeshComponents)
+void FOcclusionMaterialCameraNodeEvaluator::FindAllPrimitiveComponents(UPrimitiveComponent* InPrimitiveComponent, TSet<UPrimitiveComponent*>& OutPrimitiveComponents)
+{
+	OutPrimitiveComponents.Add(InPrimitiveComponent);
+
+	for (USceneComponent* AttachedChild : InPrimitiveComponent->GetAttachChildren())
+	{
+		if (UPrimitiveComponent* AttachedPrimitiveComponent = Cast<UPrimitiveComponent>(AttachedChild))
+		{
+			FindAllPrimitiveComponents(AttachedPrimitiveComponent, OutPrimitiveComponents);
+		}
+	}
+}
+
+void FOcclusionMaterialCameraNodeEvaluator::ApplyOcclusionMaterial(TSet<UPrimitiveComponent*> PrimitiveComponents)
 {
 	const UOcclusionMaterialCameraNode* OcclusionMaterialNode = GetCameraNodeAs<UOcclusionMaterialCameraNode>();
 	UMaterialInterface* OcclusionTransparencyMaterial = OcclusionMaterialNode->OcclusionTransparencyMaterial;
@@ -226,54 +237,55 @@ void FOcclusionMaterialCameraNodeEvaluator::ApplyOcclusionMaterial(TSet<UStaticM
 		return;
 	}
 
-	for (UStaticMeshComponent* MeshComponent : MeshComponents)
+	for (UPrimitiveComponent* PrimitiveComponent : PrimitiveComponents)
 	{
-		if (!MeshComponent || AppliedMaterialOverrides.Contains(MeshComponent))
+		if (!PrimitiveComponent || AppliedMaterialOverrides.Contains(PrimitiveComponent))
 		{
 			continue;
 		}
 
-		FOcclusionMaterialOverrideInfo& MaterialOverride = AppliedMaterialOverrides.Add(MeshComponent);
-		for (int32 MaterialIndex = 0; MaterialIndex < MeshComponent->GetNumMaterials(); ++MaterialIndex)
+		FOcclusionMaterialOverrideInfo& MaterialOverride = AppliedMaterialOverrides.Add(PrimitiveComponent);
+		const int32 NumMaterials = PrimitiveComponent->GetNumMaterials();
+		for (int32 MaterialIndex = 0; MaterialIndex < NumMaterials; ++MaterialIndex)
 		{
-			UMaterialInterface* OriginalMaterial = MeshComponent->GetMaterial(MaterialIndex);
-			UMaterialInterface* OverrideMaterial = MeshComponent->CreateDynamicMaterialInstance(MaterialIndex, OcclusionTransparencyMaterial);
+			UMaterialInterface* OriginalMaterial = PrimitiveComponent->GetMaterial(MaterialIndex);
+			UMaterialInterface* OverrideMaterial = PrimitiveComponent->CreateDynamicMaterialInstance(MaterialIndex, OcclusionTransparencyMaterial);
 			MaterialOverride.OriginalMaterials.Add(OriginalMaterial);
 			MaterialOverride.OverrideMaterials.Add(OverrideMaterial);
-			MeshComponent->SetMaterial(MaterialIndex, OverrideMaterial);
+			PrimitiveComponent->SetMaterial(MaterialIndex, OverrideMaterial);
 		}
 	}
 }
 
-void FOcclusionMaterialCameraNodeEvaluator::RemoveOcclusionMaterial(TSet<UStaticMeshComponent*> MeshComponents)
+void FOcclusionMaterialCameraNodeEvaluator::RemoveOcclusionMaterial(TSet<UPrimitiveComponent*> PrimitiveComponents)
 {
-	for (UStaticMeshComponent* MeshComponent : MeshComponents)
+	for (UPrimitiveComponent* PrimitiveComponent : PrimitiveComponents)
 	{
-		if (!MeshComponent)
+		if (!PrimitiveComponent)
 		{
 			continue;
 		}
 
 		FOcclusionMaterialOverrideInfo MaterialOverrides;
-		const bool bFoundEntry = AppliedMaterialOverrides.RemoveAndCopyValue(MeshComponent, MaterialOverrides);
+		const bool bFoundEntry = AppliedMaterialOverrides.RemoveAndCopyValue(PrimitiveComponent, MaterialOverrides);
 		if (bFoundEntry)
 		{
 			const int32 NumMaterials = MaterialOverrides.OriginalMaterials.Num();
 			for (int32 MaterialIndex = 0; MaterialIndex < NumMaterials; ++MaterialIndex)
 			{
-				MeshComponent->SetMaterial(MaterialIndex, MaterialOverrides.OriginalMaterials[MaterialIndex]);
+				PrimitiveComponent->SetMaterial(MaterialIndex, MaterialOverrides.OriginalMaterials[MaterialIndex]);
 			}
 		}
 	}
 }
 
-void FOcclusionMaterialCameraNodeEvaluator::ResolveWeakMeshComponents(TSet<TWeakObjectPtr<UStaticMeshComponent>> WeakMeshComponents, TSet<UStaticMeshComponent*>& OutMeshComponents)
+void FOcclusionMaterialCameraNodeEvaluator::ResolveWeakPrimitiveComponents(TSet<TWeakObjectPtr<UPrimitiveComponent>> WeakPrimitiveComponents, TSet<UPrimitiveComponent*>& OutPrimitiveComponents)
 {
-	for (TWeakObjectPtr<UStaticMeshComponent> WeakMeshComponent : WeakMeshComponents)
+	for (TWeakObjectPtr<UPrimitiveComponent> WeakPrimitiveComponent : WeakPrimitiveComponents)
 	{
-		if (UStaticMeshComponent* MeshComponent = WeakMeshComponent.Get())
+		if (UPrimitiveComponent* PrimitiveComponent = WeakPrimitiveComponent.Get())
 		{
-			OutMeshComponents.Add(MeshComponent);
+			OutPrimitiveComponents.Add(PrimitiveComponent);
 		}
 	}
 }
