Cameras: re-use the same material instance for occlusion material overrides

diff --git a/Engine/Plugins/Cameras/GameplayCameras/Source/GameplayCameras/Private/Nodes/Collision/OcclusionMaterialCameraNode.cpp b/Engine/Plugins/Cameras/GameplayCameras/Source/GameplayCameras/Private/Nodes/Collision/OcclusionMaterialCameraNode.cpp
--- a/Engine/Plugins/Cameras/GameplayCameras/Source/GameplayCameras/Private/Nodes/Collision/OcclusionMaterialCameraNode.cpp
+++ b/Engine/Plugins/Cameras/GameplayCameras/Source/GameplayCameras/Private/Nodes/Collision/OcclusionMaterialCameraNode.cpp
@@ -25,7 +25,6 @@ namespace UE::Cameras
 struct FOcclusionMaterialOverrideInfo
 {
 	TArray<UMaterialInterface*> OriginalMaterials;
-	TArray<UMaterialInterface*> OverrideMaterials;
 };
 
 class FOcclusionMaterialCameraNodeEvaluator : public FCameraNodeEvaluator
@@ -39,6 +38,7 @@ public:
 protected:
 
 	// FCameraNodeEvaluator interface.
+	virtual void OnAddReferencedObjects(FReferenceCollector& Collector) override;
 	virtual void OnInitialize(const FCameraNodeEvaluatorInitializeParams& Params, FCameraNodeEvaluationResult& OutResult) override;
 	virtual void OnRun(const FCameraNodeEvaluationParams& Params, FCameraNodeEvaluationResult& OutResult) override;
 
@@ -60,6 +60,7 @@ private:
 	TCameraParameterReader<FVector3d> OcclusionTargetOffsetReader;
 
 	FTraceHandle OcclusionTraceHandle;
+	TObjectPtr<UMaterialInstanceDynamic> OverrideMaterialInstance;
 	TSet<TWeakObjectPtr<UPrimitiveComponent>> CurrentlyOccludedPrimitiveComponents;
 	TMap<TWeakObjectPtr<UPrimitiveComponent>, FOcclusionMaterialOverrideInfo> AppliedMaterialOverrides;
 };
@@ -74,6 +75,11 @@ FOcclusionMaterialCameraNodeEvaluator::~FOcclusionMaterialCameraNodeEvaluator()
 	RemoveOcclusionMaterial(PrimitiveComponents);
 }
 
+void FOcclusionMaterialCameraNodeEvaluator::OnAddReferencedObjects(FReferenceCollector& Collector)
+{
+	Collector.AddReferencedObject(OverrideMaterialInstance);
+}
+
 void FOcclusionMaterialCameraNodeEvaluator::OnInitialize(const FCameraNodeEvaluatorInitializeParams& Params, FCameraNodeEvaluationResult& OutResult)
 {
 	SetNodeEvaluatorFlags(ECameraNodeEvaluatorFlags::None);
@@ -82,7 +88,12 @@ void FOcclusionMaterialCameraNodeEvaluator::OnInitialize(const FCameraNodeEvalua
 	OcclusionSphereRadiusReader.Initialize(OcclusionMaterialNode->OcclusionSphereRadius);
 	OcclusionTargetOffsetReader.Initialize(OcclusionMaterialNode->OcclusionTargetOffset);
 
-	if (!OcclusionMaterialNode->OcclusionTransparencyMaterial)
+	if (OcclusionMaterialNode->OcclusionTransparencyMaterial)
+	{
+		UObject* OuterObject = Params.EvaluationContext->GetOwner();
+		OverrideMaterialInstance = UMaterialInstanceDynamic::Create(OcclusionMaterialNode->OcclusionTransparencyMaterial, OuterObject);
+	}
+	else
 	{
 		UE_LOG(LogCameraSystem, Error, 
 				TEXT("OcclusionMaterialCameraNode: no occlusion transparency material set on '%s'"),
@@ -217,22 +228,23 @@ void FOcclusionMaterialCameraNodeEvaluator::HandleOcclusionTraceResult(UWorld* W
 
 void FOcclusionMaterialCameraNodeEvaluator::FindAllPrimitiveComponents(UPrimitiveComponent* InPrimitiveComponent, TSet<UPrimitiveComponent*>& OutPrimitiveComponents)
 {
-	OutPrimitiveComponents.Add(InPrimitiveComponent);
-
-	for (USceneComponent* AttachedChild : InPrimitiveComponent->GetAttachChildren())
+	if (InPrimitiveComponent)
 	{
-		if (UPrimitiveComponent* AttachedPrimitiveComponent = Cast<UPrimitiveComponent>(AttachedChild))
+		OutPrimitiveComponents.Add(InPrimitiveComponent);
+
+		for (USceneComponent* AttachedChild : InPrimitiveComponent->GetAttachChildren())
 		{
-			FindAllPrimitiveComponents(AttachedPrimitiveComponent, OutPrimitiveComponents);
+			if (UPrimitiveComponent* AttachedPrimitiveComponent = Cast<UPrimitiveComponent>(AttachedChild))
+			{
+				FindAllPrimitiveComponents(AttachedPrimitiveComponent, OutPrimitiveComponents);
+			}
 		}
 	}
 }
 
 void FOcclusionMaterialCameraNodeEvaluator::ApplyOcclusionMaterial(TSet<UPrimitiveComponent*> PrimitiveComponents)
 {
-	const UOcclusionMaterialCameraNode* OcclusionMaterialNode = GetCameraNodeAs<UOcclusionMaterialCameraNode>();
-	UMaterialInterface* OcclusionTransparencyMaterial = OcclusionMaterialNode->OcclusionTransparencyMaterial;
-	if (!OcclusionTransparencyMaterial)
+	if (!OverrideMaterialInstance)
 	{
 		return;
 	}
@@ -249,10 +261,8 @@ void FOcclusionMaterialCameraNodeEvaluator::ApplyOcclusionMaterial(TSet<UPrimiti
 		for (int32 MaterialIndex = 0; MaterialIndex < NumMaterials; ++MaterialIndex)
 		{
 			UMaterialInterface* OriginalMaterial = PrimitiveComponent->GetMaterial(MaterialIndex);
-			UMaterialInterface* OverrideMaterial = PrimitiveComponent->CreateDynamicMaterialInstance(MaterialIndex, OcclusionTransparencyMaterial);
 			MaterialOverride.OriginalMaterials.Add(OriginalMaterial);
-			MaterialOverride.OverrideMaterials.Add(OverrideMaterial);
-			PrimitiveComponent->SetMaterial(MaterialIndex, OverrideMaterial);
+			PrimitiveComponent->SetMaterial(MaterialIndex, OverrideMaterialInstance);
 		}
 	}
 }
