Cameras: fix possible crash when clearing the camera asset or camera rig asset on a component

diff --git a/Engine/Plugins/Cameras/GameplayCameras/Source/GameplayCameras/Private/GameFramework/GameplayCameraComponent.cpp b/Engine/Plugins/Cameras/GameplayCameras/Source/GameplayCameras/Private/GameFramework/GameplayCameraComponent.cpp
--- a/Engine/Plugins/Cameras/GameplayCameras/Source/GameplayCameras/Private/GameFramework/GameplayCameraComponent.cpp
+++ b/Engine/Plugins/Cameras/GameplayCameras/Source/GameplayCameras/Private/GameFramework/GameplayCameraComponent.cpp
@@ -114,9 +114,9 @@ void UGameplayCameraComponent::ApplyChangedParameterOverrides()
 {
 	using namespace UE::Cameras;
 
-	if (HasEvaluationContext())
+	const UCameraAsset* CameraAsset = CameraReference.GetCameraAsset();
+	if (CameraAsset && HasEvaluationContext())
 	{
-		const UCameraAsset* CameraAsset = CameraReference.GetCameraAsset();
 		FCameraNodeEvaluationResult& InitialResult = GetEvaluationContext()->GetInitialResult();
 
 		FCameraObjectInterfaceParameterOverrideHelper Helper(&InitialResult.VariableTable, &InitialResult.ContextDataTable);
diff --git a/Engine/Plugins/Cameras/GameplayCameras/Source/GameplayCameras/Private/GameFramework/GameplayCameraRigComponent.cpp b/Engine/Plugins/Cameras/GameplayCameras/Source/GameplayCameras/Private/GameFramework/GameplayCameraRigComponent.cpp
--- a/Engine/Plugins/Cameras/GameplayCameras/Source/GameplayCameras/Private/GameFramework/GameplayCameraRigComponent.cpp
+++ b/Engine/Plugins/Cameras/GameplayCameras/Source/GameplayCameras/Private/GameFramework/GameplayCameraRigComponent.cpp
@@ -160,9 +160,9 @@ void UGameplayCameraRigComponent::ApplyChangedParameterOverrides()
 {
 	using namespace UE::Cameras;
 
-	if (HasEvaluationContext())
+	const UCameraRigAsset* CameraRigAsset = CameraRigReference.GetCameraRig();
+	if (CameraRigAsset && HasEvaluationContext())
 	{
-		const UCameraRigAsset* CameraRigAsset = CameraRigReference.GetCameraRig();
 		FCameraNodeEvaluationResult& InitialResult = GetEvaluationContext()->GetInitialResult();
 
 		FCameraObjectInterfaceParameterOverrideHelper Helper(&InitialResult.VariableTable, &InitialResult.ContextDataTable);
