Cameras: use soft-object pointers for actor target info

diff --git a/Engine/Plugins/Cameras/GameplayCameras/Source/GameplayCameras/Private/Nodes/Framing/CameraActorTargetInfo.cpp b/Engine/Plugins/Cameras/GameplayCameras/Source/GameplayCameras/Private/Nodes/Framing/CameraActorTargetInfo.cpp
--- a/Engine/Plugins/Cameras/GameplayCameras/Source/GameplayCameras/Private/Nodes/Framing/CameraActorTargetInfo.cpp
+++ b/Engine/Plugins/Cameras/GameplayCameras/Source/GameplayCameras/Private/Nodes/Framing/CameraActorTargetInfo.cpp
@@ -32,9 +32,10 @@ void FCameraActorTargetInfoReader::CacheTargetInfo(const FCameraActorTargetInfo&
 		CachedTargetInfo = InTargetInfo;
 
 		CachedSkeletalMeshComponent = nullptr;
-		if (InTargetInfo.Actor && (!InTargetInfo.SocketName.IsNone() || !InTargetInfo.BoneName.IsNone()))
+		AActor* TargetActor = InTargetInfo.Actor.Get();
+		if (TargetActor && (!InTargetInfo.SocketName.IsNone() || !InTargetInfo.BoneName.IsNone()))
 		{
-			CachedSkeletalMeshComponent = InTargetInfo.Actor->FindComponentByClass<USkeletalMeshComponent>();
+			CachedSkeletalMeshComponent = TargetActor->FindComponentByClass<USkeletalMeshComponent>();
 		}
 
 		CachedBoneName = NAME_None;
@@ -74,9 +75,9 @@ bool FCameraActorTargetInfoReader::GetTargetInfo(const FCameraContextDataTable&
 		ComputeTargetBounds(OutTransform.GetLocation(), OutBounds);
 		return true;
 	}
-	else if (CachedTargetInfo.Actor)
+	else if (AActor* TargetActor = CachedTargetInfo.Actor.Get())
 	{
-		OutTransform = CachedTargetInfo.Actor->GetTransform();
+		OutTransform = TargetActor->GetTransform();
 		ComputeTargetBounds(OutTransform.GetLocation(), OutBounds);
 		return true;
 	}
@@ -101,9 +102,12 @@ void FCameraActorTargetInfoReader::ComputeTargetBounds(const FVector3d& TargetLo
 			else
 			{
 				OutBounds = FBoxSphereBounds3d(EForceInit::ForceInit);
-				if (USceneComponent* RootComponent = CachedTargetInfo.Actor->GetRootComponent())
+				if (AActor* TargetActor = CachedTargetInfo.Actor.Get())
 				{
-					OutBounds = RootComponent->Bounds;
+					if (USceneComponent* RootComponent = TargetActor->GetRootComponent())
+					{
+						OutBounds = RootComponent->Bounds;
+					}
 				}
 			}
 			break;
diff --git a/Engine/Plugins/Cameras/GameplayCameras/Source/GameplayCameras/Public/Nodes/Framing/CameraActorTargetInfo.h b/Engine/Plugins/Cameras/GameplayCameras/Source/GameplayCameras/Public/Nodes/Framing/CameraActorTargetInfo.h
--- a/Engine/Plugins/Cameras/GameplayCameras/Source/GameplayCameras/Public/Nodes/Framing/CameraActorTargetInfo.h
+++ b/Engine/Plugins/Cameras/GameplayCameras/Source/GameplayCameras/Public/Nodes/Framing/CameraActorTargetInfo.h
@@ -37,7 +37,7 @@ struct FCameraActorTargetInfo
 
 	/** The actor to target. */
 	UPROPERTY(BlueprintReadWrite, EditAnywhere, Interp, Category="Target")
-	TObjectPtr<AActor> Actor;
+	TSoftObjectPtr<AActor> Actor;
 
 	/** An optional socket to target on the actor. */
 	UPROPERTY(BlueprintReadWrite, EditAnywhere, Interp, Category="Target")
