Cameras: fix incorrect cvar and protect from division by zero in framing nodes

diff --git a/Engine/Plugins/Cameras/GameplayCameras/Source/GameplayCameras/Private/Nodes/Framing/BaseFramingCameraNode.cpp b/Engine/Plugins/Cameras/GameplayCameras/Source/GameplayCameras/Private/Nodes/Framing/BaseFramingCameraNode.cpp
--- a/Engine/Plugins/Cameras/GameplayCameras/Source/GameplayCameras/Private/Nodes/Framing/BaseFramingCameraNode.cpp
+++ b/Engine/Plugins/Cameras/GameplayCameras/Source/GameplayCameras/Private/Nodes/Framing/BaseFramingCameraNode.cpp
@@ -38,7 +38,7 @@ static FAutoConsoleVariableRef CVarFramingExtrapolationEpsilon(
 float GFramingMinDampingFactor = 0.0001;
 static FAutoConsoleVariableRef CVarFramingMinDampingFactor(
 	TEXT("GameplayCameras.Framing.MinDampingFactor"),
-	GFramingExtrapolationEpsilon,
+	GFramingMinDampingFactor,
 	TEXT("(Default: 0.0001) The minimimum reframe damping factor possible once all factors have been taken into account."));
 
 float GFramingTargetRestEpsilon = 0.001;
@@ -596,9 +596,16 @@ void FBaseFramingCameraNodeEvaluator::ComputeDesiredState(const FCameraNodeEvalu
 	NewDistanceToGo = FMath::Lerp(DistanceToGo, NewDistanceToGo, SmoothStep(State.ToggleEngageAlpha));
 
 	// Compute where we want the target this frame.
-	const FVector2d InvReframeDir(IdealToTarget / DistanceToGo);
-	const FVector2d NewScreenTarget = State.IdealTarget + InvReframeDir * NewDistanceToGo;
-	Desired.ScreenTarget = NewScreenTarget;
+	if (DistanceToGo != 0.f)
+	{
+		const FVector2d InvReframeDir(IdealToTarget / DistanceToGo);
+		const FVector2d NewScreenTarget = State.IdealTarget + InvReframeDir * NewDistanceToGo;
+		Desired.ScreenTarget = NewScreenTarget;
+	}
+	else
+	{
+		Desired.ScreenTarget = State.IdealTarget;
+	}
 
 	Desired.FramingCorrection = Desired.ScreenTarget - State.ScreenTarget;
 	Desired.bHasCorrection = true;
