Redo the previous change which was lost in a bad merge

diff --git a/Engine/Plugins/Cameras/GameplayCameras/Source/GameplayCameras/Public/Core/BaseCameraObject.h b/Engine/Plugins/Cameras/GameplayCameras/Source/GameplayCameras/Public/Core/BaseCameraObject.h
--- a/Engine/Plugins/Cameras/GameplayCameras/Source/GameplayCameras/Public/Core/BaseCameraObject.h
+++ b/Engine/Plugins/Cameras/GameplayCameras/Source/GameplayCameras/Public/Core/BaseCameraObject.h
@@ -28,7 +28,7 @@ namespace UE::Cameras
 	class ICameraObjectEventHandler
 	{
 	public:
-		virtual ~ICameraObjectEventHandler() {}
+		virtual ~ICameraObjectEventHandler() = default;
 
 		/** Called when the camera object's interface has changed. */
 		virtual void OnCameraObjectInterfaceChanged() {}
