Cameras: ensure we continue handling collision results event when we don't hit anything

This fixes a regression where absence of collision would lack updating the damper to return back to the "rest" position.

#Synced-CL 48619321

diff --git a/Engine/Plugins/Cameras/GameplayCameras/Source/GameplayCameras/Private/Nodes/Collision/CollisionPushCameraNode.cpp b/Engine/Plugins/Cameras/GameplayCameras/Source/GameplayCameras/Private/Nodes/Collision/CollisionPushCameraNode.cpp
--- a/Engine/Plugins/Cameras/GameplayCameras/Source/GameplayCameras/Private/Nodes/Collision/CollisionPushCameraNode.cpp
+++ b/Engine/Plugins/Cameras/GameplayCameras/Source/GameplayCameras/Private/Nodes/Collision/CollisionPushCameraNode.cpp
@@ -260,7 +260,7 @@ void FCollisionPushCameraNodeEvaluator::RunCollisionTrace(UWorld* World, APlayer
 	RunCollisionTrace(Params, TraceParams, OutResult, TraceResult);
 
 	CollisionTraceHandle = TraceResult.AsyncTraceHandle;
-	if (!TraceResult.HitResults.IsEmpty())
+	if (!CollisionTraceHandle.IsValid())
 	{
 		HandleCollisionTraceResult(World, TraceResult.HitResults, SafePosition, Params, OutResult);
 	}
