Cameras: don't build CDOs

#Synced-CL 48619321

diff --git a/Engine/Plugins/Cameras/GameplayCameras/Source/GameplayCameras/Private/Build/CameraAssetBuilder.cpp b/Engine/Plugins/Cameras/GameplayCameras/Source/GameplayCameras/Private/Build/CameraAssetBuilder.cpp
--- a/Engine/Plugins/Cameras/GameplayCameras/Source/GameplayCameras/Private/Build/CameraAssetBuilder.cpp
+++ b/Engine/Plugins/Cameras/GameplayCameras/Source/GameplayCameras/Private/Build/CameraAssetBuilder.cpp
@@ -9,6 +9,7 @@
 #include "Core/CameraObjectInterfaceParameterDefinition.h"
 #include "Core/CameraParameters.h"
 #include "Core/CameraRigAsset.h"
+#include "Core/CameraRigProxyAsset.h"
 #include "GameplayCamerasDelegates.h"
 #include "Logging/TokenizedMessage.h"
 
@@ -62,6 +63,11 @@ void FCameraAssetBuilder::BuildCameraImpl(bool bBuildReferencedAssets)
 		BuildLog.AddMessage(EMessageSeverity::Error, LOCTEXT("MissingDirector", "Camera has no director set."));
 	}
 
+	// Some directors might add CDOs for some reason (for instance if they use a reference collector).
+	// We need to remove them and leave them untouched.
+	AllCameraRigs.Remove(GetDefault<UCameraRigAsset>());
+	AllCameraRigProxies.Remove(GetDefault<UCameraRigProxyAsset>());
+
 	if (AllCameraRigs.IsEmpty() && AllCameraRigProxies.IsEmpty())
 	{
 		BuildLog.AddMessage(EMessageSeverity::Warning, LOCTEXT("MissingRigs", "Camera director isn't using any camera rigs or proxies."));
