Cameras: export some classes and methods to other modules

#13603
#Synced-CL 48619321

diff --git a/Engine/Plugins/Cameras/GameplayCameras/Source/GameplayCameras/Public/Core/BuiltInCameraVariables.h b/Engine/Plugins/Cameras/GameplayCameras/Source/GameplayCameras/Public/Core/BuiltInCameraVariables.h
--- a/Engine/Plugins/Cameras/GameplayCameras/Source/GameplayCameras/Public/Core/BuiltInCameraVariables.h
+++ b/Engine/Plugins/Cameras/GameplayCameras/Source/GameplayCameras/Public/Core/BuiltInCameraVariables.h
@@ -7,6 +7,8 @@
 
 #include "BuiltInCameraVariables.generated.h"
 
+#define UE_API GAMEPLAYCAMERAS_API
+
 class UDoubleCameraVariable;
 class UVector2dCameraVariable;
 
@@ -51,14 +53,14 @@ class FBuiltInCameraVariables
 public:
 
 	/** Singleton access. */
-	static const FBuiltInCameraVariables& Get();
+	UE_API static const FBuiltInCameraVariables& Get();
 
 	/** Get the definition of a built-in camera variable. */
-	const FCameraVariableDefinition& GetDefinition(EBuiltInDoubleCameraVariable BuiltInVariable) const;
+	UE_API const FCameraVariableDefinition& GetDefinition(EBuiltInDoubleCameraVariable BuiltInVariable) const;
 	/** Get the definition of a built-in camera variable. */
-	const FCameraVariableDefinition& GetDefinition(EBuiltInVector2dCameraVariable BuiltInVariable) const;
+	UE_API const FCameraVariableDefinition& GetDefinition(EBuiltInVector2dCameraVariable BuiltInVariable) const;
 	/** Get the definition of a built-in camera variable. */
-	const FCameraVariableDefinition& GetDefinition(EBuiltInRotator3dCameraVariable BuiltInVariable) const;
+	UE_API const FCameraVariableDefinition& GetDefinition(EBuiltInRotator3dCameraVariable BuiltInVariable) const;
 
 public:
 
@@ -79,3 +81,5 @@ protected:
 
 }  // namespace UE::Cameras
 
+#undef UE_API
+
diff --git a/Engine/Plugins/Cameras/GameplayCameras/Source/GameplayCameras/Public/GameFramework/IGameplayCameraSystemHost.h b/Engine/Plugins/Cameras/GameplayCameras/Source/GameplayCameras/Public/GameFramework/IGameplayCameraSystemHost.h
--- a/Engine/Plugins/Cameras/GameplayCameras/Source/GameplayCameras/Public/GameFramework/IGameplayCameraSystemHost.h
+++ b/Engine/Plugins/Cameras/GameplayCameras/Source/GameplayCameras/Public/GameFramework/IGameplayCameraSystemHost.h
@@ -11,6 +11,8 @@
 
 #include "IGameplayCameraSystemHost.generated.h"
 
+#define UE_API GAMEPLAYCAMERAS_API
+
 class UCameraRigAsset;
 class UCameraShakeAsset;
 class UCanvas;
@@ -45,7 +47,7 @@ public:
 	using FCameraEvaluationContext = UE::Cameras::FCameraEvaluationContext;
 
 	/** Gets the camera system evaluator. */
-	GAMEPLAYCAMERAS_API TSharedPtr<FCameraSystemEvaluator> GetCameraSystemEvaluator();
+	UE_API TSharedPtr<FCameraSystemEvaluator> GetCameraSystemEvaluator();
 
 	/** Returns whether a camera system evaluation has been created in this host. */
 	bool HasCameraSystem() const { return CameraSystemEvaluator.IsValid(); }
@@ -54,68 +56,68 @@ public:
 	virtual UObject* GetAsObject() { return nullptr; }
 
 	/** Returns this object as a script interface. */
-	TScriptInterface<IGameplayCameraSystemHost> GetAsScriptInterface();
+	UE_API TScriptInterface<IGameplayCameraSystemHost> GetAsScriptInterface();
 
 public:
 
 	/** Finds a valid host on the player controller's camera manager, or its view target. */
-	static IGameplayCameraSystemHost* FindActiveHost(APlayerController* PlayerController);
+	UE_API static IGameplayCameraSystemHost* FindActiveHost(APlayerController* PlayerController);
 
 protected:
 
 	/** Creates a new camera system. Asserts if there is already one. */
-	GAMEPLAYCAMERAS_API void InitializeCameraSystem();
+	UE_API void InitializeCameraSystem();
 	/** Creates a new camera system. Asserts if there is already one. */
-	GAMEPLAYCAMERAS_API void InitializeCameraSystem(const UE::Cameras::FCameraSystemEvaluatorCreateParams& Params);
+	UE_API void InitializeCameraSystem(const UE::Cameras::FCameraSystemEvaluatorCreateParams& Params);
 	/** Ensures that the camera system is created. */
-	GAMEPLAYCAMERAS_API void EnsureCameraSystemInitialized();
+	UE_API void EnsureCameraSystemInitialized();
 	/** Destroys the camera system. */
-	GAMEPLAYCAMERAS_API void DestroyCameraSystem();
+	UE_API void DestroyCameraSystem();
 
 	/** Should be called by the underlying object for garbage collection. */
-	GAMEPLAYCAMERAS_API void OnAddReferencedObjects(FReferenceCollector& Collector);
+	UE_API void OnAddReferencedObjects(FReferenceCollector& Collector);
 
 	/** Updates the camera system, if it exists. */
-	GAMEPLAYCAMERAS_API void UpdateCameraSystem(float DeltaTime);
+	UE_API void UpdateCameraSystem(float DeltaTime);
 
 #if WITH_EDITOR
 	/** Updates the camera system, if it exists, for an editor world preview. */
-	GAMEPLAYCAMERAS_API void UpdateCameraSystemForEditorPreview(float DeltaTime);
+	UE_API void UpdateCameraSystemForEditorPreview(float DeltaTime);
 #endif  // WITH_EDITOR
 
 protected:
 
 	/** Activates the given camera rig in the given layer. Should not be used with Main layer. */
-	GAMEPLAYCAMERAS_API FCameraRigInstanceID ActivateCameraRig(
+	UE_API FCameraRigInstanceID ActivateCameraRig(
 			const UCameraRigAsset* CameraRig,
 			TSharedPtr<FCameraEvaluationContext> EvaluationContext,
 			ECameraRigLayer EvaluationLayer);
 
 	/** Deactivates a previously activated camera rig. */
-	GAMEPLAYCAMERAS_API void DeactivateCameraRig(FCameraRigInstanceID InInstanceID, bool bImmediately = false);
+	UE_API void DeactivateCameraRig(FCameraRigInstanceID InInstanceID, bool bImmediately = false);
 
 	/** Starts a camera modifier rig on the given layer. */
-	GAMEPLAYCAMERAS_API FCameraRigInstanceID StartCameraModifierRig(
+	UE_API FCameraRigInstanceID StartCameraModifierRig(
 			const UCameraRigAsset* CameraRig,
 			TSharedPtr<FCameraEvaluationContext> EvaluationContext,
 			ECameraRigLayer EvaluationLayer,
 			int32 OrderKey = 0);
 
 	/** Stops a camera modifier rig on previously started on the global or visual layer. */
-	GAMEPLAYCAMERAS_API void StopCameraModifierRig(FCameraRigInstanceID InstanceID, bool bImmediately = false);
+	UE_API void StopCameraModifierRig(FCameraRigInstanceID InstanceID, bool bImmediately = false);
 
 	/** Starts a new camera shake. */
-	GAMEPLAYCAMERAS_API FCameraShakeInstanceID StartCameraShake(
+	UE_API FCameraShakeInstanceID StartCameraShake(
 			const UCameraShakeAsset* CameraShake,
 			float ShakeScale = 1.f,
 			ECameraShakePlaySpace PlaySpace = ECameraShakePlaySpace::CameraLocal,
 			FRotator UserPlaySpaceRotation = FRotator::ZeroRotator);
 
 	/** Checks if a camera shake is running. */
-	GAMEPLAYCAMERAS_API bool IsCameraShakePlaying(FCameraShakeInstanceID InInstanceID) const;
+	UE_API bool IsCameraShakePlaying(FCameraShakeInstanceID InInstanceID) const;
 
 	/** Stops a running camera shake. */
-	GAMEPLAYCAMERAS_API bool StopCameraShake(FCameraShakeInstanceID InInstanceID, bool bImmediately = false);
+	UE_API bool StopCameraShake(FCameraShakeInstanceID InInstanceID, bool bImmediately = false);
 
 protected:
 
@@ -135,3 +137,5 @@ private:
 #endif  // UE_GAMEPLAY_CAMERAS_DEBUG
 };
 
+#undef UE_API
+
diff --git a/Engine/Plugins/Cameras/GameplayCameras/Source/GameplayCameras/Public/Nodes/Input/CameraRigInput2DSlot.h b/Engine/Plugins/Cameras/GameplayCameras/Source/GameplayCameras/Public/Nodes/Input/CameraRigInput2DSlot.h
--- a/Engine/Plugins/Cameras/GameplayCameras/Source/GameplayCameras/Public/Nodes/Input/CameraRigInput2DSlot.h
+++ b/Engine/Plugins/Cameras/GameplayCameras/Source/GameplayCameras/Public/Nodes/Input/CameraRigInput2DSlot.h
@@ -13,6 +13,8 @@
 
 #include "CameraRigInput2DSlot.generated.h"
 
+#define UE_API GAMEPLAYCAMERAS_API
+
 class UVector2dCameraVariable;
 
 /**
@@ -79,13 +81,13 @@ public:
 public:
 
 	// UObject interface.
-	virtual void PostLoad() override;
+	UE_API virtual void PostLoad() override;
 
 protected:
 
 	// UCameraNode interface.
-	virtual void OnBuild(FCameraObjectBuildContext& BuildContext) override;
-	virtual FCameraNodeEvaluatorPtr OnBuildEvaluator(FCameraNodeEvaluatorBuilder& Builder) const override;
+	UE_API virtual void OnBuild(FCameraObjectBuildContext& BuildContext) override;
+	UE_API virtual FCameraNodeEvaluatorPtr OnBuildEvaluator(FCameraNodeEvaluatorBuilder& Builder) const override;
 
 private:
 
@@ -102,20 +104,20 @@ class FInput2DCameraNodeEvaluator;
 
 class FCameraRigInput2DSlotEvaluator : public FInput2DCameraNodeEvaluator
 {
-	UE_DECLARE_CAMERA_NODE_EVALUATOR_EX(GAMEPLAYCAMERAS_API, FCameraRigInput2DSlotEvaluator, FInput2DCameraNodeEvaluator)
+	UE_DECLARE_CAMERA_NODE_EVALUATOR_EX(UE_API, FCameraRigInput2DSlotEvaluator, FInput2DCameraNodeEvaluator)
 
 public:
 
-	FCameraRigInput2DSlotEvaluator();
+	UE_API FCameraRigInput2DSlotEvaluator();
 
 protected:
 
 	// FCameraNodeEvaluator interface.
-	virtual void OnInitialize(const FCameraNodeEvaluatorInitializeParams& Params, FCameraNodeEvaluationResult& OutResult) override;
-	virtual void OnUpdateParameters(const FCameraBlendedParameterUpdateParams& Params, FCameraBlendedParameterUpdateResult& OutResult) override;
-	virtual void OnRun(const FCameraNodeEvaluationParams& Params, FCameraNodeEvaluationResult& OutResult) override;
-	virtual void OnExecuteOperation(const FCameraOperationParams& Params, FCameraOperation& Operation) override;
-	virtual void OnSerialize(const FCameraNodeEvaluatorSerializeParams& Params, FArchive& Ar) override;
+	UE_API virtual void OnInitialize(const FCameraNodeEvaluatorInitializeParams& Params, FCameraNodeEvaluationResult& OutResult) override;
+	UE_API virtual void OnUpdateParameters(const FCameraBlendedParameterUpdateParams& Params, FCameraBlendedParameterUpdateResult& OutResult) override;
+	UE_API virtual void OnRun(const FCameraNodeEvaluationParams& Params, FCameraNodeEvaluationResult& OutResult) override;
+	UE_API virtual void OnExecuteOperation(const FCameraOperationParams& Params, FCameraOperation& Operation) override;
+	UE_API virtual void OnSerialize(const FCameraNodeEvaluatorSerializeParams& Params, FArchive& Ar) override;
 
 protected:
 
@@ -129,3 +131,5 @@ protected:
 
 }  // namespace UE::Cameras
 
+#undef UE_API
+
diff --git a/Engine/Plugins/Cameras/GameplayCameras/Source/GameplayCameras/Public/Nodes/Input/CameraRigInputSlotTypes.h b/Engine/Plugins/Cameras/GameplayCameras/Source/GameplayCameras/Public/Nodes/Input/CameraRigInputSlotTypes.h
--- a/Engine/Plugins/Cameras/GameplayCameras/Source/GameplayCameras/Public/Nodes/Input/CameraRigInputSlotTypes.h
+++ b/Engine/Plugins/Cameras/GameplayCameras/Source/GameplayCameras/Public/Nodes/Input/CameraRigInputSlotTypes.h
@@ -6,6 +6,8 @@
 
 #include "CameraRigInputSlotTypes.generated.h"
 
+#define UE_API GAMEPLAYCAMERAS_API
+
 /**
  * General options for an input slot.
  */
@@ -44,14 +46,14 @@ struct FCameraParameterClamping
 public:
 
 	/** Clamps the given value. */
-	double ClampValue(double Value) const;
+	UE_API double ClampValue(double Value) const;
 
 	/** 
 	 * Gets the effective min/max values for this struct.
 	 * If a bound is disabled, the effective value will be lowest or max
 	 * double precision values.
 	 */
-	void GetEffectiveClamping(double& OutMinValue, double& OutMaxValue) const;
+	UE_API void GetEffectiveClamping(double& OutMinValue, double& OutMaxValue) const;
 };
 
 /**
@@ -71,6 +73,8 @@ struct FCameraParameterNormalization
 public:
 
 	/** Normalizes the given value. */
-	double NormalizeValue(double Value) const;
+	UE_API double NormalizeValue(double Value) const;
 };
 
+#undef UE_API
+
