Cameras: unify how Guids are generated for various assets

#Synced-CL 48619321

diff --git a/Engine/Plugins/Cameras/GameplayCameras/Source/GameplayCameras/Private/Core/BaseCameraObject.cpp b/Engine/Plugins/Cameras/GameplayCameras/Source/GameplayCameras/Private/Core/BaseCameraObject.cpp
--- a/Engine/Plugins/Cameras/GameplayCameras/Source/GameplayCameras/Private/Core/BaseCameraObject.cpp
+++ b/Engine/Plugins/Cameras/GameplayCameras/Source/GameplayCameras/Private/Core/BaseCameraObject.cpp
@@ -24,9 +24,35 @@ void UBaseCameraObject::PostLoad()
 {
 	Super::PostLoad();
 
+	if (!Guid.IsValid())
+	{
+		Guid = FGuid::NewGuid();
+	}
+
 #if UE_VERSION_OLDER_THAN(5,8,0)
 	using namespace UE::Cameras;
 	FCameraObjectInterfaceParameterBuilder::FixUpDefaultParameterProperties(ParameterDefinitions, DefaultParameters);
 #endif
 }
 
+void UBaseCameraObject::PostInitProperties()
+{
+	Super::PostInitProperties();
+
+	if (!HasAnyFlags(RF_ClassDefaultObject | RF_ArchetypeObject | RF_NeedLoad | RF_WasLoaded) && 
+			!Guid.IsValid())
+	{
+		Guid = FGuid::NewGuid();
+	}
+}
+
+void UBaseCameraObject::PostDuplicate(EDuplicateMode::Type DuplicateMode)
+{
+	Super::PostDuplicate(DuplicateMode);
+
+	if (DuplicateMode == EDuplicateMode::Normal)
+	{
+		Guid = FGuid::NewGuid();
+	}
+}
+
diff --git a/Engine/Plugins/Cameras/GameplayCameras/Source/GameplayCameras/Private/Core/CameraRigAsset.cpp b/Engine/Plugins/Cameras/GameplayCameras/Source/GameplayCameras/Private/Core/CameraRigAsset.cpp
--- a/Engine/Plugins/Cameras/GameplayCameras/Source/GameplayCameras/Private/Core/CameraRigAsset.cpp
+++ b/Engine/Plugins/Cameras/GameplayCameras/Source/GameplayCameras/Private/Core/CameraRigAsset.cpp
@@ -51,11 +51,6 @@ void UCameraRigAsset::PostLoad()
 
 #endif
 
-	if (!Guid.IsValid())
-	{
-		Guid = FGuid::NewGuid();
-	}
-
 	// Initialize the ParameterType of blendable parameters that pre-existed the new interface data.
 	// The type defaults to Boolean so check only those.
 	for (UCameraObjectInterfaceBlendableParameter* BlendableParameter : Interface.BlendableParameters)
@@ -71,27 +66,6 @@ void UCameraRigAsset::PostLoad()
 	}
 }
 
-void UCameraRigAsset::PostInitProperties()
-{
-	Super::PostInitProperties();
-
-	if (!HasAnyFlags(RF_ClassDefaultObject | RF_ArchetypeObject | RF_NeedLoad | RF_WasLoaded) && 
-			!Guid.IsValid())
-	{
-		Guid = FGuid::NewGuid();
-	}
-}
-
-void UCameraRigAsset::PostDuplicate(EDuplicateMode::Type DuplicateMode)
-{
-	Super::PostDuplicate(DuplicateMode);
-
-	if (DuplicateMode == EDuplicateMode::Normal)
-	{
-		Guid = FGuid::NewGuid();
-	}
-}
-
 void UCameraRigAsset::GetOwnedGameplayTags(FGameplayTagContainer& TagContainer) const
 {
 	TagContainer.AppendTags(GameplayTags);
diff --git a/Engine/Plugins/Cameras/GameplayCameras/Source/GameplayCameras/Private/Core/CameraRigProxyAsset.cpp b/Engine/Plugins/Cameras/GameplayCameras/Source/GameplayCameras/Private/Core/CameraRigProxyAsset.cpp
--- a/Engine/Plugins/Cameras/GameplayCameras/Source/GameplayCameras/Private/Core/CameraRigProxyAsset.cpp
+++ b/Engine/Plugins/Cameras/GameplayCameras/Source/GameplayCameras/Private/Core/CameraRigProxyAsset.cpp
@@ -6,16 +6,26 @@
 
 #include UE_INLINE_GENERATED_CPP_BY_NAME(CameraRigProxyAsset)
 
-UCameraRigProxyAsset::UCameraRigProxyAsset(const FObjectInitializer& ObjectInit)
-	: Super(ObjectInit)
+void UCameraRigProxyAsset::PostLoad()
 {
-	if (!IsTemplate())
+	Super::PostLoad();
+
+	if (!Guid.IsValid())
 	{
 		Guid = FGuid::NewGuid();
 	}
 }
 
-#if WITH_EDITORONLY_DATA
+void UCameraRigProxyAsset::PostInitProperties()
+{
+	Super::PostInitProperties();
+	
+	if (!HasAnyFlags(RF_ClassDefaultObject | RF_ArchetypeObject | RF_NeedLoad | RF_WasLoaded) && 
+			!Guid.IsValid())
+	{
+		Guid = FGuid::NewGuid();
+	}
+}
 
 void UCameraRigProxyAsset::PostDuplicate(EDuplicateMode::Type DuplicateMode)
 {
@@ -27,5 +37,3 @@ void UCameraRigProxyAsset::PostDuplicate(EDuplicateMode::Type DuplicateMode)
 	}
 }
 
-#endif
-
diff --git a/Engine/Plugins/Cameras/GameplayCameras/Source/GameplayCameras/Private/Core/CameraShakeAsset.cpp b/Engine/Plugins/Cameras/GameplayCameras/Source/GameplayCameras/Private/Core/CameraShakeAsset.cpp
--- a/Engine/Plugins/Cameras/GameplayCameras/Source/GameplayCameras/Private/Core/CameraShakeAsset.cpp
+++ b/Engine/Plugins/Cameras/GameplayCameras/Source/GameplayCameras/Private/Core/CameraShakeAsset.cpp
@@ -4,34 +4,13 @@
 
 #include "Build/CameraBuildLog.h"
 #include "Build/CameraShakeAssetBuilder.h"
-#include "Core/ShakeCameraNode.h"
+#include "Core/ShakeCameraNode.h"  // IWYU pragma: keep
 #include "Misc/EngineVersionComparison.h"
 #include "UObject/AssetRegistryTagsContext.h"
 #include "UObject/ObjectSaveContext.h"
 
 #include UE_INLINE_GENERATED_CPP_BY_NAME(CameraShakeAsset)
 
-void UCameraShakeAsset::PostInitProperties()
-{
-	Super::PostInitProperties();
-
-	if (!HasAnyFlags(RF_ClassDefaultObject | RF_ArchetypeObject | RF_NeedLoad | RF_WasLoaded) && 
-			!Guid.IsValid())
-	{
-		Guid = FGuid::NewGuid();
-	}
-}
-
-void UCameraShakeAsset::PostDuplicate(EDuplicateMode::Type DuplicateMode)
-{
-	Super::PostDuplicate(DuplicateMode);
-
-	if (DuplicateMode == EDuplicateMode::Normal)
-	{
-		Guid = FGuid::NewGuid();
-	}
-}
-
 void UCameraShakeAsset::BuildCameraShake()
 {
 	using namespace UE::Cameras;
diff --git a/Engine/Plugins/Cameras/GameplayCameras/Source/GameplayCameras/Public/Core/BaseCameraObject.h b/Engine/Plugins/Cameras/GameplayCameras/Source/GameplayCameras/Public/Core/BaseCameraObject.h
--- a/Engine/Plugins/Cameras/GameplayCameras/Source/GameplayCameras/Public/Core/BaseCameraObject.h
+++ b/Engine/Plugins/Cameras/GameplayCameras/Source/GameplayCameras/Public/Core/BaseCameraObject.h
@@ -92,6 +92,9 @@ public:
 
 public:
 
+	/** Gets the object's unique ID. */
+	const FGuid& GetGuid() const { return Guid; }
+
 	/** Gets the default values for the parameters exposed on this camera rig. */
 	const FInstancedPropertyBag& GetDefaultParameters() const { return DefaultParameters; }
 
@@ -110,9 +113,15 @@ public:
 
 	// UObject interface.
 	virtual void PostLoad() override;
+	virtual void PostInitProperties() override;
+	virtual void PostDuplicate(EDuplicateMode::Type DuplicateMode) override;
 
 private:
 
+	/** The camera object's unique ID. */
+	UPROPERTY()
+	FGuid Guid;
+
 	/** The default interface parameter values, generated during build. */
 	UPROPERTY()
 	FInstancedPropertyBag DefaultParameters;
diff --git a/Engine/Plugins/Cameras/GameplayCameras/Source/GameplayCameras/Public/Core/CameraRigAsset.h b/Engine/Plugins/Cameras/GameplayCameras/Source/GameplayCameras/Public/Core/CameraRigAsset.h
--- a/Engine/Plugins/Cameras/GameplayCameras/Source/GameplayCameras/Public/Core/CameraRigAsset.h
+++ b/Engine/Plugins/Cameras/GameplayCameras/Source/GameplayCameras/Public/Core/CameraRigAsset.h
@@ -91,9 +91,6 @@ public:
 	UPROPERTY(EditAnywhere, Category="Transition")
 	ECameraRigInitialOrientation InitialOrientation = ECameraRigInitialOrientation::None;
 
-	/** Gets the camera rig's unique ID. */
-	const FGuid& GetGuid() const { return Guid; }
-
 public:
 
 	/** The current build state of this camera rig. */
@@ -154,8 +151,6 @@ protected:
 
 	// UObject interface
 	virtual void PostLoad() override;
-	virtual void PostInitProperties() override;
-	virtual void PostDuplicate(EDuplicateMode::Type DuplicateMode) override;
 	virtual void PreSave(FObjectPreSaveContext ObjectSaveContext) override;
 	virtual void GetAssetRegistryTags(FAssetRegistryTagsContext Context) const override;
 #if WITH_EDITOR
@@ -164,10 +159,6 @@ protected:
 
 private:
 
-	/** The camera rig's unique ID. */
-	UPROPERTY()
-	FGuid Guid;
-
 #if WITH_EDITORONLY_DATA
 
 	/** Position of the camera rig node in the node graph editor. */
diff --git a/Engine/Plugins/Cameras/GameplayCameras/Source/GameplayCameras/Public/Core/CameraRigProxyAsset.h b/Engine/Plugins/Cameras/GameplayCameras/Source/GameplayCameras/Public/Core/CameraRigProxyAsset.h
--- a/Engine/Plugins/Cameras/GameplayCameras/Source/GameplayCameras/Public/Core/CameraRigProxyAsset.h
+++ b/Engine/Plugins/Cameras/GameplayCameras/Source/GameplayCameras/Public/Core/CameraRigProxyAsset.h
@@ -21,18 +21,17 @@ class UCameraRigProxyAsset : public UObject
 
 public:
 
-	UCameraRigProxyAsset(const FObjectInitializer& ObjectInit);
-
-public:
+	/** Gets the camera rig proxy's unique ID. */
+	const FGuid& GetGuid() const { return Guid; }
 
 	// UObject interface.
-#if WITH_EDITORONLY_DATA
+	virtual void PostLoad() override;
+	virtual void PostInitProperties() override;
 	virtual void PostDuplicate(EDuplicateMode::Type DuplicateMode) override;
-#endif
 
-public:
+private:
 
-	/** Unique identifier for this camera rig proxy. */
+	/** The camera rig proxy's unique ID. */
 	UPROPERTY()
 	FGuid Guid;
 };
diff --git a/Engine/Plugins/Cameras/GameplayCameras/Source/GameplayCameras/Public/Core/CameraShakeAsset.h b/Engine/Plugins/Cameras/GameplayCameras/Source/GameplayCameras/Public/Core/CameraShakeAsset.h
--- a/Engine/Plugins/Cameras/GameplayCameras/Source/GameplayCameras/Public/Core/CameraShakeAsset.h
+++ b/Engine/Plugins/Cameras/GameplayCameras/Source/GameplayCameras/Public/Core/CameraShakeAsset.h
@@ -49,11 +49,6 @@ public:
 	UPROPERTY(EditAnywhere, Category="Common")
 	bool bIsSingleInstance = false;
 
-public:
-
-	/** Gets the camera shake's unique ID. */
-	const FGuid& GetGuid() const { return Guid; }
-
 public:
 
 	/** The current build state of this camera shake. */
@@ -95,16 +90,10 @@ protected:
 #endif
 
 	// UObject interface.
-	virtual void PostInitProperties() override;
-	virtual void PostDuplicate(EDuplicateMode::Type DuplicateMode) override;
 	virtual void PreSave(FObjectPreSaveContext ObjectSaveContext) override;
 
 private:
 
-	/** The camera shake's unique ID. */
-	UPROPERTY()
-	FGuid Guid;
-
 #if WITH_EDITORONLY_DATA
 
 	/** Position of the camera shake node in the node graph editor. */
