Cameras: fix harmless but confusing copy/paste typo

#Synced-CL 48619321

diff --git a/Engine/Plugins/Cameras/GameplayCameras/Source/GameplayCameras/Private/MovieScene/MovieSceneGameplayCamerasComponentTypes.cpp b/Engine/Plugins/Cameras/GameplayCameras/Source/GameplayCameras/Private/MovieScene/MovieSceneGameplayCamerasComponentTypes.cpp
--- a/Engine/Plugins/Cameras/GameplayCameras/Source/GameplayCameras/Private/MovieScene/MovieSceneGameplayCamerasComponentTypes.cpp
+++ b/Engine/Plugins/Cameras/GameplayCameras/Source/GameplayCameras/Private/MovieScene/MovieSceneGameplayCamerasComponentTypes.cpp
@@ -28,8 +28,8 @@ namespace UE::MovieScene
 namespace UE::Cameras
 {
 
-static bool GMovieSceneUMGComponentTypesDestroyed = false;
-static TUniquePtr<FMovieSceneGameplayCamerasComponentTypes> GMovieSceneUMGComponentTypes;
+static bool GGameplayCamerasComponentTypesDestroyed = false;
+static TUniquePtr<FMovieSceneGameplayCamerasComponentTypes> GGameplayCamerasComponentTypes;
 
 FMovieSceneGameplayCamerasComponentTypes::FMovieSceneGameplayCamerasComponentTypes()
 {
@@ -58,18 +58,18 @@ FMovieSceneGameplayCamerasComponentTypes::~FMovieSceneGameplayCamerasComponentTy
 
 void FMovieSceneGameplayCamerasComponentTypes::Destroy()
 {
-	GMovieSceneUMGComponentTypes.Reset();
-	GMovieSceneUMGComponentTypesDestroyed = true;
+	GGameplayCamerasComponentTypes.Reset();
+	GGameplayCamerasComponentTypesDestroyed = true;
 }
 
 FMovieSceneGameplayCamerasComponentTypes* FMovieSceneGameplayCamerasComponentTypes::Get()
 {
-	if (!GMovieSceneUMGComponentTypes.IsValid())
+	if (!GGameplayCamerasComponentTypes.IsValid())
 	{
-		check(!GMovieSceneUMGComponentTypesDestroyed);
-		GMovieSceneUMGComponentTypes.Reset(new FMovieSceneGameplayCamerasComponentTypes);
+		check(!GGameplayCamerasComponentTypesDestroyed);
+		GGameplayCamerasComponentTypes.Reset(new FMovieSceneGameplayCamerasComponentTypes);
 	}
-	return GMovieSceneUMGComponentTypes.Get();
+	return GGameplayCamerasComponentTypes.Get();
 }
 
 } // namespace UE::Cameras
