Cameras: add option for freezing previous camera rigs when activating a new one

#Synced-CL 48619321

diff --git a/Engine/Plugins/Cameras/GameplayCameras/Source/GameplayCameras/Private/Core/TransientBlendStackCameraNode.cpp b/Engine/Plugins/Cameras/GameplayCameras/Source/GameplayCameras/Private/Core/TransientBlendStackCameraNode.cpp
--- a/Engine/Plugins/Cameras/GameplayCameras/Source/GameplayCameras/Private/Core/TransientBlendStackCameraNode.cpp
+++ b/Engine/Plugins/Cameras/GameplayCameras/Source/GameplayCameras/Private/Core/TransientBlendStackCameraNode.cpp
@@ -75,6 +75,15 @@ FBlendStackEntryID FTransientBlendStackCameraNodeEvaluator::Push(const FBlendSta
 
 FBlendStackEntryID FTransientBlendStackCameraNodeEvaluator::PushNewEntry(const FBlendStackCameraPushParams& Params, const UCameraRigTransition* Transition)
 {
+	// Freeze all previous entries if requested.
+	if (Transition && Transition->bFreezePreviousCameraRigs)
+	{
+		for (FCameraRigEntry& Entry : Entries)
+		{
+			FreezeEntry(Entry);
+		}
+	}
+
 	// Create the new root node to wrap the new camera rig's root node, and the specific
 	// blend node for this transition.
 	// We need to const-cast here to be able to use our own blend stack node as the outer
@@ -132,6 +141,16 @@ FBlendStackEntryID FTransientBlendStackCameraNodeEvaluator::PushNewEntry(const F
 
 FBlendStackEntryID FTransientBlendStackCameraNodeEvaluator::PushMergedEntry(const FBlendStackCameraPushParams& PushParams, const UCameraRigTransition* Transition)
 {
+	// Freeze all previous entries if requested, except the one we're merging into.
+	if (Transition && Transition->bFreezePreviousCameraRigs)
+	{
+		for (int32 Index = 0; Index < Entries.Num() - 1; ++Index)
+		{
+			FCameraRigEntry& Entry = Entries[Index];
+			FreezeEntry(Entry);
+		}
+	}
+
 	const UBlendCameraNode* Blend = Transition ? Transition->Blend.Get() : nullptr;
 
 	FCameraRigEntry& TopEntry = Entries.Top();
diff --git a/Engine/Plugins/Cameras/GameplayCameras/Source/GameplayCameras/Public/Core/CameraRigTransition.h b/Engine/Plugins/Cameras/GameplayCameras/Source/GameplayCameras/Public/Core/CameraRigTransition.h
--- a/Engine/Plugins/Cameras/GameplayCameras/Source/GameplayCameras/Public/Core/CameraRigTransition.h
+++ b/Engine/Plugins/Cameras/GameplayCameras/Source/GameplayCameras/Public/Core/CameraRigTransition.h
@@ -163,6 +163,10 @@ public:
 	UPROPERTY(EditAnywhere, Category="Transition")
 	bool bOverrideInitialOrientation = false;
 
+	/** Freeze all previously running camera rigs. */
+	UPROPERTY(EditAnywhere, Category="Transition")
+	bool bFreezePreviousCameraRigs = false;
+
 	/**
 	 * Whether this transition allows merging two similar camera rigs together.
 	 * Similar camera rigs run the same underlying camera rig prefab with different parameter
