Cameras: activate and deactivate the output camera component along with its parent

This fixes problems where a deactivated Gameplay Camera/Camera Rig component is still picked up by the engine when its owning actor is set as the view target, even though there's another active camera component on that actor.

#Synced-CL 48619321

diff --git a/Engine/Plugins/Cameras/GameplayCameras/Source/GameplayCameras/Private/GameFramework/GameplayCameraComponentBase.cpp b/Engine/Plugins/Cameras/GameplayCameras/Source/GameplayCameras/Private/GameFramework/GameplayCameraComponentBase.cpp
--- a/Engine/Plugins/Cameras/GameplayCameras/Source/GameplayCameras/Private/GameFramework/GameplayCameraComponentBase.cpp
+++ b/Engine/Plugins/Cameras/GameplayCameras/Source/GameplayCameras/Private/GameFramework/GameplayCameraComponentBase.cpp
@@ -188,6 +188,11 @@ bool UGameplayCameraComponentBase::EnsureCameraSystemHostIfNeeded()
 	FCameraEvaluationContextStack& ContextStack = HostedEvaluator->GetEvaluationContextStack();
 	ContextStack.PushContext(EvaluationContext.ToSharedRef());
 
+	if (ensure(OutputCameraComponent))
+	{
+		OutputCameraComponent->Activate();
+	}
+
 	return true;
 }
 
@@ -197,6 +202,7 @@ void UGameplayCameraComponentBase::DestroyCameraSystemHost()
 	if (OutputCameraComponent)
 	{
 		OutputCameraComponent->SetRelativeTransform(FTransform());
+		OutputCameraComponent->Deactivate();
 	}
 
 	DestroyCameraSystem();
@@ -445,6 +451,7 @@ void UGameplayCameraComponentBase::OnRegister()
 	if (OutputCameraComponent == nullptr)
 	{
 		OutputCameraComponent = NewObject<UCineCameraComponent>(this, TEXT("OutputCameraComponent"), RF_Transient | RF_TextExportTransient);
+		OutputCameraComponent->SetAutoActivate(false);
 		OutputCameraComponent->SetupAttachment(this);
 		OutputCameraComponent->CreationMethod = CreationMethod;
 		OutputCameraComponent->RegisterComponentWithWorld(GetWorld());
