Cameras: fix incorrect FOV when using a vanilla camera component or actor as the view target inside a GPC camera manager

When setting the FOV or focal length on a camera pose, we always need to invalidate the other so that both aren't defined at the same time.

#Synced-CL 48619321

diff --git a/Engine/Plugins/Cameras/GameplayCameras/Source/GameplayCameras/Private/GameFramework/ActorCameraEvaluationContext.cpp b/Engine/Plugins/Cameras/GameplayCameras/Source/GameplayCameras/Private/GameFramework/ActorCameraEvaluationContext.cpp
--- a/Engine/Plugins/Cameras/GameplayCameras/Source/GameplayCameras/Private/GameFramework/ActorCameraEvaluationContext.cpp
+++ b/Engine/Plugins/Cameras/GameplayCameras/Source/GameplayCameras/Private/GameFramework/ActorCameraEvaluationContext.cpp
@@ -69,6 +69,7 @@ void FActorCameraEvaluationContext::ApplyMinimalViewInfo(const FMinimalViewInfo&
 	CameraPose.SetLocation(ViewInfo.Location);
 	CameraPose.SetRotation(ViewInfo.Rotation);
 	CameraPose.SetFieldOfView(ViewInfo.FOV);
+	CameraPose.SetFocalLength(-1.f);
 	CameraPose.SetNearClippingPlane(ViewInfo.PerspectiveNearClipPlane);
 	CameraPose.SetSensorWidth(CameraPose.GetSensorHeight() * ViewInfo.AspectRatio);
 	CameraPose.SetAspectRatioAxisConstraint(ViewInfo.AspectRatioAxisConstraint.Get(CameraPose.GetAspectRatioAxisConstraint()));
