Cameras: don't render two debug infos when a GPC camera manager is running with a nested standalone GPC camera component

Since the camera manager and the view target were both considered as the site for the "default" camera system, the camera debugger would end up drawing them on top of each other.
Now we only consider the view target as default if the camera manager doesn't already run a camera system itself.

#Synced-CL 48619321

diff --git a/Engine/Plugins/Cameras/GameplayCameras/Source/GameplayCameras/Private/Core/CameraSystemEvaluator.cpp b/Engine/Plugins/Cameras/GameplayCameras/Source/GameplayCameras/Private/Core/CameraSystemEvaluator.cpp
--- a/Engine/Plugins/Cameras/GameplayCameras/Source/GameplayCameras/Private/Core/CameraSystemEvaluator.cpp
+++ b/Engine/Plugins/Cameras/GameplayCameras/Source/GameplayCameras/Private/Core/CameraSystemEvaluator.cpp
@@ -603,7 +603,7 @@ void FCameraSystemEvaluator::DebugUpdate(const FCameraSystemDebugUpdateParams& P
 #endif
 
 	FRootCameraDebugDrawParams RootParams;
-	RootParams.bIsCameraManagerOrViewTarget = Params.bIsCameraManagerOrViewTarget;
+	RootParams.bFromDefaultCameraSystem = Params.bIsDefaultCameraSystem;
 	RootParams.bForceDraw = Params.bForceDraw;
 	FCameraDebugRenderer Renderer(OwnerWorld, Params.CanvasObject, Params.bIsDebugCameraEnabled);
 	RootDebugBlock->RootDebugDraw(RootParams, Renderer);
diff --git a/Engine/Plugins/Cameras/GameplayCameras/Source/GameplayCameras/Private/Debug/RootCameraDebugBlock.cpp b/Engine/Plugins/Cameras/GameplayCameras/Source/GameplayCameras/Private/Debug/RootCameraDebugBlock.cpp
--- a/Engine/Plugins/Cameras/GameplayCameras/Source/GameplayCameras/Private/Debug/RootCameraDebugBlock.cpp
+++ b/Engine/Plugins/Cameras/GameplayCameras/Source/GameplayCameras/Private/Debug/RootCameraDebugBlock.cpp
@@ -149,7 +149,7 @@ void FRootCameraDebugBlock::BuildDebugBlocks(const FCameraSystemEvaluator& Camer
 
 void FRootCameraDebugBlock::RootDebugDraw(const FRootCameraDebugDrawParams& Params, FCameraDebugRenderer& Renderer)
 {
-	if (!ShouldDebugDraw(DebugID, Params.bIsCameraManagerOrViewTarget))
+	if (!ShouldDebugDraw(DebugID, Params.bFromDefaultCameraSystem))
 	{
 		return;
 	}
@@ -175,7 +175,7 @@ void FRootCameraDebugBlock::OnSerialize(FArchive& Ar)
 	Ar << DebugID;
 }
 
-bool FRootCameraDebugBlock::ShouldDebugDraw(FCameraSystemDebugID InDebugID, bool bIsActive)
+bool FRootCameraDebugBlock::ShouldDebugDraw(FCameraSystemDebugID InDebugID, bool bIsDefaultCameraSystem)
 {
 	if (!GGameplayCamerasDebugEnable)
 	{
@@ -191,7 +191,7 @@ bool FRootCameraDebugBlock::ShouldDebugDraw(FCameraSystemDebugID InDebugID, bool
 	// we are in charge of the "active" camera, whatever that means.
 	const bool bDoDebugDraw = (WantedDebugID == InDebugID) 
 		|| WantedDebugID.IsAny()
-		|| (WantedDebugID.IsAuto() && (bIsActive || NumCameraSystems == 1));
+		|| (WantedDebugID.IsAuto() && (bIsDefaultCameraSystem || NumCameraSystems == 1));
 	return bDoDebugDraw;
 }
 
diff --git a/Engine/Plugins/Cameras/GameplayCameras/Source/GameplayCameras/Private/GameFramework/IGameplayCameraSystemHost.cpp b/Engine/Plugins/Cameras/GameplayCameras/Source/GameplayCameras/Private/GameFramework/IGameplayCameraSystemHost.cpp
--- a/Engine/Plugins/Cameras/GameplayCameras/Source/GameplayCameras/Private/GameFramework/IGameplayCameraSystemHost.cpp
+++ b/Engine/Plugins/Cameras/GameplayCameras/Source/GameplayCameras/Private/GameFramework/IGameplayCameraSystemHost.cpp
@@ -14,6 +14,7 @@
 #include "Debug/DebugDrawService.h"
 #include "Engine/Canvas.h"
 #include "Engine/World.h"
+#include "GameFramework/GameplayCamerasPlayerCameraManager.h"
 #include "GameFramework/PlayerController.h"
 #include "Services/CameraModifierService.h"
 #include "Services/CameraShakeService.h"
@@ -308,17 +309,22 @@ void IGameplayCameraSystemHost::DebugDraw(UCanvas* Canvas, APlayerController* Pl
 		const AActor* ViewTarget = ActualPlayerController ? ActualPlayerController->GetViewTarget() : nullptr;
 		const bool bThisIsCameraManager = (ActualPlayerController && ThisAsObject == ActualPlayerController->PlayerCameraManager);
 		const bool bThisIsViewTarget = (ThisAsObject && ViewTarget && ThisAsObject->GetTypedOuter<AActor>() == ViewTarget);
+		const bool bHasCameraManager = (ActualPlayerController && ActualPlayerController->PlayerCameraManager->IsA<AGameplayCamerasPlayerCameraManager>());
 
 		// We're looking from the outside if we are not the view target, or if we don't have a player
 		// anymore (which happens in spectator mode like with the debug camera).
-		bool bIsDebugCameraEnabled = (
+		const bool bIsDebugCameraEnabled = (
 				(!bThisIsCameraManager && !bThisIsViewTarget) ||
 				!ActualPlayerController || !ActualPlayerController->Player);
 
+		// The default camera system is the one running the camera manager or, if there none there, the 
+		// one running the current view target.
+		const bool bIsDefaultCameraSystem = (bThisIsCameraManager || (bThisIsViewTarget && !bHasCameraManager));
+
 		FCameraSystemDebugUpdateParams DebugUpdateParams;
 		DebugUpdateParams.CanvasObject = Canvas;
 		DebugUpdateParams.bIsDebugCameraEnabled = bIsDebugCameraEnabled;
-		DebugUpdateParams.bIsCameraManagerOrViewTarget = (bThisIsCameraManager || bThisIsViewTarget);
+		DebugUpdateParams.bIsDefaultCameraSystem = bIsDefaultCameraSystem;
 		CameraSystemEvaluator->DebugUpdate(DebugUpdateParams);
 	}
 }
diff --git a/Engine/Plugins/Cameras/GameplayCameras/Source/GameplayCameras/Public/Core/CameraSystemEvaluator.h b/Engine/Plugins/Cameras/GameplayCameras/Source/GameplayCameras/Public/Core/CameraSystemEvaluator.h
--- a/Engine/Plugins/Cameras/GameplayCameras/Source/GameplayCameras/Public/Core/CameraSystemEvaluator.h
+++ b/Engine/Plugins/Cameras/GameplayCameras/Source/GameplayCameras/Public/Core/CameraSystemEvaluator.h
@@ -134,8 +134,11 @@ struct FCameraSystemDebugUpdateParams
 	/** Whether the debug camera is enabled, giving an "outside" view of camera system. */
 	bool bIsDebugCameraEnabled = false;
 
-	/** Whether this camera system is run by the active camera manager or view target. */
-	bool bIsCameraManagerOrViewTarget = false;
+	/**
+	 * Whether this camera system should be considered the "default" one, such as when
+	 * running inside the active camera manager or the view target.
+	 */
+	bool bIsDefaultCameraSystem = false;
 
 	/** Whether to force drawing debug info for this camera system. */
 	bool bForceDraw = false;
diff --git a/Engine/Plugins/Cameras/GameplayCameras/Source/GameplayCameras/Public/Debug/RootCameraDebugBlock.h b/Engine/Plugins/Cameras/GameplayCameras/Source/GameplayCameras/Public/Debug/RootCameraDebugBlock.h
--- a/Engine/Plugins/Cameras/GameplayCameras/Source/GameplayCameras/Public/Debug/RootCameraDebugBlock.h
+++ b/Engine/Plugins/Cameras/GameplayCameras/Source/GameplayCameras/Public/Debug/RootCameraDebugBlock.h
@@ -25,10 +25,10 @@ GAMEPLAYCAMERAS_API extern FString GGameplayCamerasDebugCategories;
 struct FRootCameraDebugDrawParams
 {
 	/**
-	 * Whether the debug blocks belong to a camera system running as the camera manager
-	 * of active view target.
+	 * Whether the debug blocks belong to a camera system that is considered the "default" one,
+	 * such as the one running inside the camera manager, or the active view target.
 	 */
-	bool bIsCameraManagerOrViewTarget = false;
+	bool bFromDefaultCameraSystem = false;
 
 	/** Whether to force draw the debug blocks. */
 	bool bForceDraw = false;
@@ -55,7 +55,7 @@ public:
 public:
 
 	/** Checks whether the debug data with the given ID would be drawn on screen. */
-	GAMEPLAYCAMERAS_API static bool ShouldDebugDraw(FCameraSystemDebugID InDebugID, bool bIsActive);
+	GAMEPLAYCAMERAS_API static bool ShouldDebugDraw(FCameraSystemDebugID InDebugID, bool bIsDefaultCameraSystem);
 
 protected:
 
