Cameras: fix bugs related to enums in the data table

- Store enums with an int32 for handling large flags enums.
- Fix a bug where we have to convert the value (not the pointer) of a data table entry in the case of enums, which may have different backing types.
- Also fix incorrect data traits for strings.

#Synced-CL 48619321

diff --git a/Engine/Plugins/Cameras/GameplayCameras/Source/GameplayCameras/Private/Core/CameraContextDataTable.cpp b/Engine/Plugins/Cameras/GameplayCameras/Source/GameplayCameras/Private/Core/CameraContextDataTable.cpp
--- a/Engine/Plugins/Cameras/GameplayCameras/Source/GameplayCameras/Private/Core/CameraContextDataTable.cpp
+++ b/Engine/Plugins/Cameras/GameplayCameras/Source/GameplayCameras/Private/Core/CameraContextDataTable.cpp
@@ -197,8 +197,8 @@ bool FCameraContextDataTable::GetDataTypeAllocationInfo(ECameraContextDataType D
 			OutAlignOf = alignof(FString);
 			break;
 		case ECameraContextDataType::Enum:
-			OutSizeOf = sizeof(uint8);
-			OutAlignOf = alignof(uint8);
+			OutSizeOf = sizeof(uint32);
+			OutAlignOf = alignof(uint32);
 			break;
 		case ECameraContextDataType::Struct:
 			{
@@ -255,7 +255,7 @@ bool FCameraContextDataTable::ConstructDataValue(ECameraContextDataType DataType
 				const UEnum* EnumType = Cast<const UEnum>(DataTypeObject);
 				if (ensure(EnumType))
 				{
-					*reinterpret_cast<uint8*>(DataPtr) = (uint8)EnumType->GetValueByIndex(0);
+					*reinterpret_cast<uint32*>(DataPtr) = (uint32)EnumType->GetValueByIndex(0);
 				}
 			}
 			break;
@@ -407,7 +407,7 @@ const FName& FCameraContextDataTable::GetNameData(FCameraContextDataID InID) con
 
 const FString& FCameraContextDataTable::GetStringData(FCameraContextDataID InID) const
 {
-	if (const FString* Value = GetDataImpl<FString>(InID, ECameraContextDataType::Name, nullptr))
+	if (const FString* Value = GetDataImpl<FString>(InID, ECameraContextDataType::String, nullptr))
 	{
 		return *Value;
 	}
@@ -416,9 +416,9 @@ const FString& FCameraContextDataTable::GetStringData(FCameraContextDataID InID)
 	return DefaultValue;
 }
 
-uint8 FCameraContextDataTable::GetEnumData(FCameraContextDataID InID, const UEnum* EnumType) const
+uint32 FCameraContextDataTable::GetEnumData(FCameraContextDataID InID, const UEnum* EnumType) const
 {
-	if (const uint8* Value = GetDataImpl<uint8>(InID, ECameraContextDataType::Enum, EnumType))
+	if (const uint32* Value = GetDataImpl<uint32>(InID, ECameraContextDataType::Enum, EnumType))
 	{
 		return *Value;
 	}
@@ -450,7 +450,7 @@ FInstancedStruct FCameraContextDataTable::GetInstancedStructData(FCameraContextD
 
 UObject* FCameraContextDataTable::GetObjectData(FCameraContextDataID InID) const
 {
-	if (const TObjectPtr<UObject>* Value = GetDataImpl<TObjectPtr<UObject>>(InID, ECameraContextDataType::Name, nullptr))
+	if (const TObjectPtr<UObject>* Value = GetDataImpl<TObjectPtr<UObject>>(InID, ECameraContextDataType::Object, nullptr))
 	{
 		return Value->Get();
 	}
@@ -459,7 +459,7 @@ UObject* FCameraContextDataTable::GetObjectData(FCameraContextDataID InID) const
 
 UClass* FCameraContextDataTable::GetClassData(FCameraContextDataID InID) const
 {
-	if (const TObjectPtr<UClass>* Value = GetDataImpl<TObjectPtr<UClass>>(InID, ECameraContextDataType::Name, nullptr))
+	if (const TObjectPtr<UClass>* Value = GetDataImpl<TObjectPtr<UClass>>(InID, ECameraContextDataType::Class, nullptr))
 	{
 		return Value->Get();
 	}
@@ -476,7 +476,7 @@ void FCameraContextDataTable::SetStringData(FCameraContextDataID InID, const FSt
 	SetDataImpl(InID, ECameraContextDataType::String, nullptr, InData);
 }
 
-void FCameraContextDataTable::SetEnumData(FCameraContextDataID InID, const UEnum* EnumType, uint8 InData)
+void FCameraContextDataTable::SetEnumData(FCameraContextDataID InID, const UEnum* EnumType, uint32 InData)
 {
 	SetDataImpl(InID, ECameraContextDataType::Enum, EnumType, InData);
 }
@@ -533,7 +533,7 @@ void FCameraContextDataTable::SetStringArrayData(FCameraContextDataID InID, TCon
 	SetArrayDataImpl(InID, ECameraContextDataType::String, nullptr, InData);
 }
 
-void FCameraContextDataTable::SetEnumArrayData(FCameraContextDataID InID, const UEnum* EnumType, TConstArrayView<uint8> InData)
+void FCameraContextDataTable::SetEnumArrayData(FCameraContextDataID InID, const UEnum* EnumType, TConstArrayView<uint32> InData)
 {
 	SetArrayDataImpl(InID, ECameraContextDataType::Enum, EnumType, InData);
 }
@@ -842,7 +842,7 @@ bool FCameraContextDataTable::SetDataValue(ECameraContextDataType DataType, cons
 			*reinterpret_cast<FString*>(DestDataPtr) = *reinterpret_cast<const FString*>(SrcDataPtr);
 			break;
 		case ECameraContextDataType::Enum:
-			*reinterpret_cast<uint8*>(DestDataPtr) = *reinterpret_cast<const uint8*>(SrcDataPtr);
+			*reinterpret_cast<uint32*>(DestDataPtr) = *reinterpret_cast<const uint32*>(SrcDataPtr);
 			break;
 		case ECameraContextDataType::Struct:
 			{
diff --git a/Engine/Plugins/Cameras/GameplayCameras/Source/GameplayCameras/Public/Core/CameraContextDataReader.h b/Engine/Plugins/Cameras/GameplayCameras/Source/GameplayCameras/Public/Core/CameraContextDataReader.h
--- a/Engine/Plugins/Cameras/GameplayCameras/Source/GameplayCameras/Public/Core/CameraContextDataReader.h
+++ b/Engine/Plugins/Cameras/GameplayCameras/Source/GameplayCameras/Public/Core/CameraContextDataReader.h
@@ -38,14 +38,27 @@ struct TCameraContextDataReader
 		{
 			return *DefaultValuePtr;
 		}
+
+		if constexpr(std::is_enum_v<DataType>)
+		{
+			// Enums need to be converted: they are always stored as uint32 in the table,
+			// but they may actually have a different backing type. If that type is shorter
+			// or longer, we can't reinterpret_cast the memory, we need to convert the
+			// uint32 value instead.
+			if (const uint8* RawValue = ContextDataTable.TryGetData(
+						DataID, ECameraContextDataType::Enum, StaticEnum<DataType>()))
+			{
+				return DataType(*reinterpret_cast<const uint32*>(RawValue));
+			}
+		}
 		else
 		{
 			if (const DataType* ActualValue = ContextDataTable.TryGetData<DataType>(DataID))
 			{
 				return *ActualValue;
 			}
-			return *DefaultValuePtr;
 		}
+		return *DefaultValuePtr;
 	}
 
 private:
diff --git a/Engine/Plugins/Cameras/GameplayCameras/Source/GameplayCameras/Public/Core/CameraContextDataTable.h b/Engine/Plugins/Cameras/GameplayCameras/Source/GameplayCameras/Public/Core/CameraContextDataTable.h
--- a/Engine/Plugins/Cameras/GameplayCameras/Source/GameplayCameras/Public/Core/CameraContextDataTable.h
+++ b/Engine/Plugins/Cameras/GameplayCameras/Source/GameplayCameras/Public/Core/CameraContextDataTable.h
@@ -68,7 +68,7 @@ public:
 
 	const FName& GetNameData(FCameraContextDataID InID) const;
 	const FString& GetStringData(FCameraContextDataID InID) const;
-	uint8 GetEnumData(FCameraContextDataID InID, const UEnum* EnumType) const;
+	uint32 GetEnumData(FCameraContextDataID InID, const UEnum* EnumType) const;
 	FConstStructView GetStructViewData(FCameraContextDataID InID, const UScriptStruct* StructType) const;
 	FInstancedStruct GetInstancedStructData(FCameraContextDataID InID, const UScriptStruct* StructType) const;
 	UObject* GetObjectData(FCameraContextDataID InID) const;
@@ -102,7 +102,7 @@ public:
 
 	void SetNameData(FCameraContextDataID InID, const FName& InData);
 	void SetStringData(FCameraContextDataID InID, const FString& InData);
-	void SetEnumData(FCameraContextDataID InID, const UEnum* EnumType, uint8 InData);
+	void SetEnumData(FCameraContextDataID InID, const UEnum* EnumType, uint32 InData);
 	void SetObjectData(FCameraContextDataID InID, UObject* InData);
 	void SetClassData(FCameraContextDataID InID, UClass* InData);
 
@@ -117,7 +117,7 @@ public:
 
 	void SetNameArrayData(FCameraContextDataID InID, TConstArrayView<FName> InData);
 	void SetStringArrayData(FCameraContextDataID InID, TConstArrayView<FString> InData);
-	void SetEnumArrayData(FCameraContextDataID InID, const UEnum* EnumType, TConstArrayView<uint8> InData);
+	void SetEnumArrayData(FCameraContextDataID InID, const UEnum* EnumType, TConstArrayView<uint32> InData);
 	void SetObjectArrayData(FCameraContextDataID InID, TConstArrayView<UObject*> InData);
 	void SetClassArrayData(FCameraContextDataID InID, TConstArrayView<UClass*> InData);
 
@@ -308,7 +308,7 @@ ENUM_CLASS_FLAGS(FCameraContextDataTable::EEntryFlags)
 template<typename EnumType>
 EnumType FCameraContextDataTable::GetEnumData(FCameraContextDataID InID) const
 {
-	if (const uint8* Value = GetDataImpl<uint8>(InID, ECameraContextDataType::Enum, StaticEnum<EnumType>()))
+	if (const uint32* Value = GetDataImpl<uint32>(InID, ECameraContextDataType::Enum, StaticEnum<EnumType>()))
 	{
 		return EnumType(*Value);
 	}
@@ -398,7 +398,7 @@ const ValueType* FCameraContextDataTable::TryGetArrayData(FCameraContextDataID I
 template<typename EnumType>
 void FCameraContextDataTable::SetEnumData(FCameraContextDataID InID, EnumType InData)
 {
-	SetDataImpl(InID, ECameraContextDataType::Enum, StaticEnum<EnumType>(), InData);
+	SetDataImpl(InID, ECameraContextDataType::Enum, StaticEnum<EnumType>(), (uint32)InData);
 }
 
 template<typename StructType>
@@ -500,7 +500,7 @@ ECameraContextDataType TCameraContextDataTraits<DataType>::GetDataType()
 	}
 	else if constexpr(std::is_same_v<DataType, FString>)
 	{
-		return ECameraContextDataType::Name;
+		return ECameraContextDataType::String;
 	}
 	else if constexpr(std::is_enum_v<DataType>)
 	{
