Cameras: make the search feature in camera editors based on the edited graphs and not the underlying object model

Previously, searching through camera node or transition graphs was done based on the underlying object trees. Now, we more simply search the graphs that have been created out of these objects (thanks to UObjectTreeGraph). This makes it possible to get search results that are only indirectly related to the underlying data model, such as extra nodes and pins that show up in the editor.

#Synced-CL 48619321

diff --git a/Engine/Plugins/Cameras/GameplayCameras/Source/GameplayCamerasEditor/Private/Editors/ObjectTreeGraphSearch.cpp b/Engine/Plugins/Cameras/GameplayCameras/Source/GameplayCamerasEditor/Private/Editors/ObjectTreeGraphSearch.cpp
--- a/Engine/Plugins/Cameras/GameplayCameras/Source/GameplayCamerasEditor/Private/Editors/ObjectTreeGraphSearch.cpp
+++ b/Engine/Plugins/Cameras/GameplayCameras/Source/GameplayCamerasEditor/Private/Editors/ObjectTreeGraphSearch.cpp
@@ -2,149 +2,82 @@
 
 #include "Editors/ObjectTreeGraphSearch.h"
 
-#include "Core/ObjectTreeGraphRootObject.h"
-#include "Editors/ObjectTreeGraphConfig.h"
+#include "EdGraph/EdGraph.h"
+#include "EdGraph/EdGraphNode.h"
+#include "EdGraph/EdGraphPin.h"
+#include "EdGraph/EdGraphSchema.h"
+
+#define LOCTEXT_NAMESPACE "ObjectTreeGraphSearch"
 
 FObjectTreeGraphSearch::FObjectTreeGraphSearch()
 {
 }
 
-void FObjectTreeGraphSearch::AddRootObject(UObject* InObject, const FObjectTreeGraphConfig* InGraphConfig)
+void FObjectTreeGraphSearch::AddGraph(UEdGraph* InGraph)
 {
-	RootObjectInfos.Add(FRootObjectInfo{ InObject, InGraphConfig });
+	SearchGraphs.Add(InGraph);
+}
+
+void FObjectTreeGraphSearch::SetGraphs(TConstArrayView<UEdGraph*> InGraphs)
+{
+	SearchGraphs = InGraphs;
 }
 
 void FObjectTreeGraphSearch::Search(TArrayView<FString> InTokens, TArray<FObjectTreeGraphSearchResult>& OutResults) const
 {
-	for (const FRootObjectInfo& RootObjectInfo : RootObjectInfos)
+	for (UEdGraph* Graph : SearchGraphs)
 	{
-		SearchRootObject(RootObjectInfo, InTokens, OutResults);
+		FSearchState State;
+		State.Graph = Graph;
+		State.Tokens = InTokens;
+
+		SearchGraph(Graph, State);
+
+		OutResults.Append(State.Results);
 	}
 }
 
-void FObjectTreeGraphSearch::SearchRootObject(const FRootObjectInfo& InRootObjectInfo, TArrayView<FString> InTokens, TArray<FSearchResult>& OutResults) const
+void FObjectTreeGraphSearch::SearchGraph(UEdGraph* Graph, FSearchState& State) const
 {
-	UObject* RootObject = InRootObjectInfo.WeakRootObject.Get();
-	if (!RootObject || !InRootObjectInfo.GraphConfig)
+	if (!Graph)
 	{
 		return;
 	}
 
-	FSearchState State;
-	State.GraphConfig = InRootObjectInfo.GraphConfig;
-	State.Tokens = InTokens;
-	State.RootObject = RootObject;
-	State.ObjectStack.Add(RootObject);
-
-	if (IObjectTreeGraphRootObject* RootObjectInterface = Cast<IObjectTreeGraphRootObject>(RootObject))
+	for (UEdGraphNode* Node : Graph->Nodes)
 	{
-		TSet<UObject*> ConnectableObjects;
-		RootObjectInterface->GetConnectableObjects(State.GraphConfig->GraphName, ConnectableObjects);
-		for (UObject* ConnectableObject : ConnectableObjects)
-		{
-			State.ObjectStack.Add(ConnectableObject);
-		}
+		SearchNode(Node, State);
 	}
-
-	while (State.ObjectStack.Num() > 0)
-	{
-		UObject* CurObject = State.ObjectStack.Pop();
-		if (CurObject && !State.VisitedObjects.Contains(CurObject))
-		{
-			State.VisitedObjects.Add(CurObject);
-
-			SearchObject(CurObject, State);
-		}
-	}
-
-	OutResults.Append(State.Results);
 }
 
-void FObjectTreeGraphSearch::SearchObject(UObject* InObject, FSearchState& InOutState) const
+void FObjectTreeGraphSearch::SearchNode(UEdGraphNode* Node, FSearchState& State) const
 {
-	UClass* ObjectClass = InObject->GetClass();
-	const FObjectTreeGraphConfig& GraphConfig(*InOutState.GraphConfig);
+	if (!Node)
+	{
+		return;
+	}
 
-	if (MatchObject(InObject, InOutState))
+	const FText NodeTitle = Node->GetNodeTitle(ENodeTitleType::FullTitle);
+	if (MatchString(NodeTitle.ToString(), State))
 	{
-		InOutState.Results.Add(FSearchResult{ InOutState.RootObject, InOutState.GraphConfig, InObject });
+		State.Results.Add(FSearchResult(State.Graph, Node));
 	}
 
-	for (TFieldIterator<FProperty> PropertyIt(ObjectClass); PropertyIt; ++PropertyIt)
+	for (UEdGraphPin* Pin : Node->Pins)
 	{
-		if (FObjectProperty* ObjectProperty = CastField<FObjectProperty>(*PropertyIt))
+		if (!Pin)
 		{
-			if (!GraphConfig.IsConnectable(ObjectProperty))
-			{
-				continue;
-			}
-
-			if (MatchObjectProperty(InObject, *PropertyIt, InOutState))
-			{
-				InOutState.Results.Add(
-						FSearchResult{ InOutState.RootObject, InOutState.GraphConfig, InObject, PropertyIt->GetFName() });
-			}
-
-			TObjectPtr<UObject> ConnectedObject;
-			ObjectProperty->GetValue_InContainer(InObject, &ConnectedObject);
-			if (ConnectedObject)
-			{
-				InOutState.ObjectStack.Add(ConnectedObject);
-			}
+			continue;
 		}
-		else if (FArrayProperty* ArrayProperty = CastField<FArrayProperty>(*PropertyIt))
+
+		const FText PinName = Pin->GetDisplayName();
+		if (MatchString(PinName.ToString(), State))
 		{
-			if (!GraphConfig.IsConnectable(ArrayProperty))
-			{
-				continue;
-			}
-
-			if (MatchObjectProperty(InObject, *PropertyIt, InOutState))
-			{
-				InOutState.Results.Add(
-						FSearchResult{ InOutState.RootObject, InOutState.GraphConfig, InObject, PropertyIt->GetFName() });
-			}
-
-			FObjectProperty* InnerProperty = CastFieldChecked<FObjectProperty>(ArrayProperty->Inner);
-			FScriptArrayHelper ArrayHelper(ArrayProperty, ArrayProperty->ContainerPtrToValuePtr<void>(InObject));
-
-			const int32 ArrayNum = ArrayHelper.Num();
-			for (int32 Index = 0; Index < ArrayNum; ++Index)
-			{
-				UObject* ConnectedObject = InnerProperty->GetObjectPropertyValue(ArrayHelper.GetRawPtr(Index));
-				if (ConnectedObject)
-				{
-					InOutState.ObjectStack.Add(ConnectedObject);
-				}
-			}
+			State.Results.Add(FSearchResult(State.Graph, Node, Pin->GetFName()));
 		}
 	}
 }
 
-bool FObjectTreeGraphSearch::MatchObject(UObject* InObject, const FSearchState& InState) const
-{
-	const FObjectTreeGraphConfig& GraphConfig(*InState.GraphConfig);
-
-	const FString DisplayNameText = GraphConfig.GetDisplayNameText(InObject).ToString();
-	if (MatchString(DisplayNameText, InState))
-	{
-		return true;
-	}
-
-	return false;
-}
-
-bool FObjectTreeGraphSearch::MatchObjectProperty(UObject* InObject, FProperty* InProperty, const FSearchState& InState) const
-{
-	const FString PropertyName = InProperty->GetName();
-	if (MatchString(PropertyName, InState))
-	{
-		return true;
-	}
-
-	return false;
-}
-
 bool FObjectTreeGraphSearch::MatchString(const FString& InString, const FSearchState& InState) const
 {
 	if (InString.IsEmpty())
@@ -163,3 +96,5 @@ bool FObjectTreeGraphSearch::MatchString(const FString& InString, const FSearchS
 	return true;
 }
 
+#undef LOCTEXT_NAMESPACE
+
diff --git a/Engine/Plugins/Cameras/GameplayCameras/Source/GameplayCamerasEditor/Private/Editors/SCameraRigAssetEditor.cpp b/Engine/Plugins/Cameras/GameplayCameras/Source/GameplayCamerasEditor/Private/Editors/SCameraRigAssetEditor.cpp
--- a/Engine/Plugins/Cameras/GameplayCameras/Source/GameplayCamerasEditor/Private/Editors/SCameraRigAssetEditor.cpp
+++ b/Engine/Plugins/Cameras/GameplayCameras/Source/GameplayCamerasEditor/Private/Editors/SCameraRigAssetEditor.cpp
@@ -234,6 +234,24 @@ void SCameraRigAssetEditor::FocusHome()
 	FindAndJumpToObjectNode(CameraRigAsset, CurrentMode);
 }
 
+bool SCameraRigAssetEditor::JumpToNode(UEdGraphNode* InNode, FName InPinName)
+{
+	UEdGraph* OwningGraph = InNode->GetGraph();
+	if (OwningGraph == NodeGraph)
+	{
+		SetEditorMode(ECameraRigAssetEditorMode::NodeGraph);
+		NodeGraphEditor->JumpToNode(InNode, InPinName);
+		return true;
+	}
+	else if (OwningGraph == TransitionGraph)
+	{
+		SetEditorMode(ECameraRigAssetEditorMode::TransitionGraph);
+		TransitionGraphEditor->JumpToNode(InNode, InPinName);
+		return true;
+	}
+	return false;
+}
+
 bool SCameraRigAssetEditor::FindAndJumpToObjectNode(UObject* InObject)
 {
 	if (FindAndJumpToObjectNode(InObject, ECameraRigAssetEditorMode::NodeGraph))
@@ -280,7 +298,7 @@ FText SCameraRigAssetEditor::GetCameraRigAssetName(UObjectTreeGraph* ForGraph) c
 {
 	if (CameraRigAsset && ForGraph)
 	{
-		const FObjectTreeGraphConfig& GraphConfig = TransitionGraph->GetConfig();
+		const FObjectTreeGraphConfig& GraphConfig = ForGraph->GetConfig();
 		return GraphConfig.GetDisplayNameText(CameraRigAsset);
 	}
 	return LOCTEXT("NoCameraRig", "No Camera Rig");
diff --git a/Engine/Plugins/Cameras/GameplayCameras/Source/GameplayCamerasEditor/Private/Editors/SCameraRigTransitionEditor.cpp b/Engine/Plugins/Cameras/GameplayCameras/Source/GameplayCamerasEditor/Private/Editors/SCameraRigTransitionEditor.cpp
--- a/Engine/Plugins/Cameras/GameplayCameras/Source/GameplayCamerasEditor/Private/Editors/SCameraRigTransitionEditor.cpp
+++ b/Engine/Plugins/Cameras/GameplayCameras/Source/GameplayCamerasEditor/Private/Editors/SCameraRigTransitionEditor.cpp
@@ -148,11 +148,11 @@ void SCameraRigTransitionEditor::FocusHome()
 	}
 }
 
-void SCameraRigTransitionEditor::JumpToNode(UEdGraphNode* InGraphNode)
+void SCameraRigTransitionEditor::JumpToNode(UEdGraphNode* InGraphNode, FName InPinName)
 {
 	if (InGraphNode)
 	{
-		TransitionGraphEditor->JumpToNode(InGraphNode);
+		TransitionGraphEditor->JumpToNode(InGraphNode, InPinName);
 	}
 }
 
diff --git a/Engine/Plugins/Cameras/GameplayCameras/Source/GameplayCamerasEditor/Private/Editors/SFindInObjectTreeGraph.cpp b/Engine/Plugins/Cameras/GameplayCameras/Source/GameplayCamerasEditor/Private/Editors/SFindInObjectTreeGraph.cpp
--- a/Engine/Plugins/Cameras/GameplayCameras/Source/GameplayCamerasEditor/Private/Editors/SFindInObjectTreeGraph.cpp
+++ b/Engine/Plugins/Cameras/GameplayCameras/Source/GameplayCamerasEditor/Private/Editors/SFindInObjectTreeGraph.cpp
@@ -2,11 +2,10 @@
 
 #include "Editors/SFindInObjectTreeGraph.h"
 
-#include "Core/ObjectTreeGraphObject.h"
 #include "EdGraph/EdGraph.h"
 #include "EdGraph/EdGraphNode.h"
-#include "EdGraph/EdGraphPin.h"
 #include "EdGraph/EdGraphSchema.h"
+#include "Editors/ObjectTreeGraph.h"
 #include "Editors/ObjectTreeGraphConfig.h"
 #include "Editors/ObjectTreeGraphSearch.h"
 #include "Framework/Application/SlateApplication.h"
@@ -25,23 +24,11 @@ FFindInObjectTreeGraphResult::FFindInObjectTreeGraphResult(const FText& InCustom
 
 FFindInObjectTreeGraphResult::FFindInObjectTreeGraphResult(
 		TSharedPtr<FFindInObjectTreeGraphResult>& InParent, 
-		const FFindInObjectTreeGraphSource& InSource, 
-		UObject* InObject)
+		UEdGraphNode* InNode, 
+		FName InPinName)
 	: Parent(InParent)
-	, WeakObject(InObject)
-	, Source(InSource)
-{
-}
-
-FFindInObjectTreeGraphResult::FFindInObjectTreeGraphResult(
-		TSharedPtr<FFindInObjectTreeGraphResult>& InParent, 
-		const FFindInObjectTreeGraphSource& InSource, 
-		UObject* InObject, 
-		FName InPropertyName)
-	: Parent(InParent)
-	, WeakObject(InObject)
-	, PropertyName(InPropertyName)
-	, Source(InSource)
+	, WeakNode(InNode)
+	, PinName(InPinName)
 {
 }
 
@@ -50,26 +37,18 @@ TSharedRef<SWidget>	FFindInObjectTreeGraphResult::GetIcon() const
 	FSlateColor IconColor = FSlateColor::UseForeground();
 	const FSlateBrush* Brush = NULL;
 
-	UObject* Object = WeakObject.Get();
+	UEdGraphNode* Node = WeakNode.Get();
 
-	if (Object && !PropertyName.IsNone())
+	if (Node && !PinName.IsNone())
 	{
-		UClass* ObjectClass = Object->GetClass();
-		FProperty* Property = ObjectClass->FindPropertyByName(PropertyName);
-		if (Property->IsA<FArrayProperty>())
-		{
-			Brush = FAppStyle::GetBrush(TEXT("GraphEditor.ArrayPinIcon"));
-		}
-		else if (Property->IsA<FObjectProperty>())
-		{
-			Brush = FAppStyle::GetBrush(TEXT("GraphEditor.RefPinIcon"));
-		}
-		else
+		if (UEdGraphPin* Pin = Node->FindPin(PinName))
 		{
+			const UEdGraphSchema* Schema = Node->GetSchema();
+			IconColor = Schema->GetPinTypeColor(Pin->PinType);
 			Brush = FAppStyle::GetBrush(TEXT("GraphEditor.PinIcon"));
 		}
 	}
-	else if (Object)
+	else if (Node)
 	{
 		Brush = FAppStyle::GetBrush(TEXT("GraphEditor.NodeGlyph"));
 	}
@@ -82,9 +61,9 @@ TSharedRef<SWidget>	FFindInObjectTreeGraphResult::GetIcon() const
 
 FText FFindInObjectTreeGraphResult::GetCategory() const
 {
-	if (WeakObject.IsValid())
+	if (WeakNode.IsValid())
 	{
-		if (PropertyName.IsNone())
+		if (PinName.IsNone())
 		{
 			return LOCTEXT("NodeCategory", "Node");
 		}
@@ -98,19 +77,20 @@ FText FFindInObjectTreeGraphResult::GetCategory() const
 
 FText FFindInObjectTreeGraphResult::GetText() const
 {
-	UObject* Object = WeakObject.Get();
-	if (Object)
+	UEdGraphNode* Node = WeakNode.Get();
+	if (Node)
 	{
-		if (PropertyName.IsNone())
+		if (PinName.IsNone())
 		{
-			const FText DisplayNameText = Source.GraphConfig->GetDisplayNameText(Object);
-			return DisplayNameText;
+			return Node->GetNodeTitle(ENodeTitleType::FullTitle);
+		}
+		else if (UEdGraphPin* Pin = Node->FindPin(PinName))
+		{
+			return Pin->GetDisplayName();
 		}
 		else
 		{
-			UClass* ObjectClass = Object->GetClass();
-			FProperty* Property = ObjectClass->FindPropertyByName(PropertyName);
-			return Property->GetDisplayNameText();
+			return FText::FromName(PinName);
 		}
 	}
 	return CustomText;
@@ -118,18 +98,18 @@ FText FFindInObjectTreeGraphResult::GetText() const
 
 FText FFindInObjectTreeGraphResult::GetCommentText() const
 {
-	if (IObjectTreeGraphObject* ObjectInterface = Cast<IObjectTreeGraphObject>(WeakObject.Get()))
+	if (UEdGraphNode* Node = WeakNode.Get())
 	{
-		return FText::FromString(ObjectInterface->GetGraphNodeCommentText(Source.GraphConfig->GraphName));
+		return FText::FromString(Node->NodeComment);
 	}
 	return FText::GetEmpty();
 }
 
 FReply FFindInObjectTreeGraphResult::OnClick(TSharedRef<SFindInObjectTreeGraph> FindInObjectTreeGraph)
 {
-	if (UObject* Object = WeakObject.Get())
+	if (UEdGraphNode* Node = WeakNode.Get())
 	{
-		FindInObjectTreeGraph->OnJumpToObjectRequested.ExecuteIfBound(Object, PropertyName);
+		FindInObjectTreeGraph->OnJumpToNodeRequested.ExecuteIfBound(Node, PinName);
 		return FReply::Handled();
 	}
 	return FReply::Unhandled();
@@ -137,8 +117,8 @@ FReply FFindInObjectTreeGraphResult::OnClick(TSharedRef<SFindInObjectTreeGraph>
 
 void SFindInObjectTreeGraph::Construct(const FArguments& InArgs)
 {
-	OnGetRootObjectsToSearch = InArgs._OnGetRootObjectsToSearch;
-	OnJumpToObjectRequested = InArgs._OnJumpToObjectRequested;
+	OnGetGraphsToSearch = InArgs._OnGetGraphsToSearch;
+	OnJumpToNodeRequested = InArgs._OnJumpToNodeRequested;
 
 	ChildSlot
 	[
@@ -180,6 +160,13 @@ void SFindInObjectTreeGraph::FocusSearchEditBox()
 	FSlateApplication::Get().SetKeyboardFocus(SearchBox, EFocusCause::SetDirectly);
 }
 
+void SFindInObjectTreeGraph::Search(const FString& InSearchQuery)
+{
+	SearchBox->SetText(FText::FromString(InSearchQuery));
+	SearchQuery = InSearchQuery;
+	StartSearch();
+}
+
 void SFindInObjectTreeGraph::OnSearchTextChanged(const FText& Text)
 {
 	SearchQuery = Text.ToString();
@@ -276,68 +263,79 @@ void SFindInObjectTreeGraph::StartSearch()
 		HighlightText = FText::FromString(SearchQuery);
 
 		TArray<FFindInObjectTreeGraphSource> Sources;
-		OnGetRootObjectsToSearch.ExecuteIfBound(Sources);
+		OnGetGraphsToSearch.ExecuteIfBound(Sources);
 
 		FObjectTreeGraphSearch Searcher;
 		for (const FFindInObjectTreeGraphSource& Source : Sources)
 		{
-			Searcher.AddRootObject(Source.RootObject, Source.GraphConfig);
+			Searcher.AddGraph(Source.Graph);
 		}
 
 		Searcher.Search(Tokens, SearchResults);
 	}
-
+	
 	// Convert simple flat search results to hierarchical results, where property results
 	// are always under an object result, object results are always under a graph result, 
 	// and so on. We do this by creating blank results if we need to, although it makes
 	// the assumption that the flat results come ordered (e.g. an object result won't be
 	// found after a property result for that object).
 	{
-		TMap<UObject*, FResultPtr> RootObjectToWidgetResult;
-		TMap<UObject*, FResultPtr> ObjectToWidgetResult;
+		TMap<UEdGraph*, FResultPtr> RootGraphToWidgetResult;
+		TMap<UEdGraphNode*, FResultPtr> NodeToWidgetResult;
 		for (const FObjectTreeGraphSearchResult& SearchResult : SearchResults)
 		{
-			FFindInObjectTreeGraphSource CurSource{ SearchResult.RootObject, SearchResult.GraphConfig };
-
 			FResultPtr GraphResult;
-			if (SearchResult.RootObject)
+			if (SearchResult.Graph)
 			{
-				GraphResult = RootObjectToWidgetResult.FindRef(SearchResult.RootObject);
+				GraphResult = RootGraphToWidgetResult.FindRef(SearchResult.Graph);
 				if (!GraphResult)
 				{
-					const FText RootObjectDisplayText = SearchResult.GraphConfig->GetDisplayNameText(SearchResult.RootObject);
-					const FText& GraphDisplayName = SearchResult.GraphConfig->GraphDisplayInfo.DisplayName;
-					const FText GraphResultText = FText::Format(
-							LOCTEXT("GraphResultFmt", "{0}: {1}"), { RootObjectDisplayText, GraphDisplayName });
+					FGraphDisplayInfo GraphDisplayInfo;
+					const UEdGraphSchema* Schema = SearchResult.Graph->GetSchema();
+					Schema->GetGraphDisplayInformation(*SearchResult.Graph, GraphDisplayInfo);
+
+					FText GraphResultText;
+					if (UObjectTreeGraph* ObjectTreeGraph = Cast<UObjectTreeGraph>(SearchResult.Graph))
+					{
+						const UObject* RootObject = ObjectTreeGraph->GetRootObject();
+						const FText RootObjectText = ObjectTreeGraph->GetConfig().GetDisplayNameText(RootObject);
+						GraphResultText = FText::Format(
+								LOCTEXT("GraphResultFmt", "{0}: {1}"), { RootObjectText, GraphDisplayInfo.DisplayName });
+					}
+					else
+					{
+						GraphResultText = GraphDisplayInfo.DisplayName;
+					}
+
 					GraphResult = MakeShared<FFindInObjectTreeGraphResult>(GraphResultText);
-					RootObjectToWidgetResult.Add(SearchResult.RootObject, GraphResult);
+					RootGraphToWidgetResult.Add(SearchResult.Graph, GraphResult);
 					Results.Add(GraphResult);
 				}
 			}
 
-			FResultPtr ObjectResult;
-			if (SearchResult.Object)
+			FResultPtr NodeResult;
+			if (SearchResult.Node)
 			{
-				ObjectResult = ObjectToWidgetResult.FindRef(SearchResult.Object);
-				if (!ObjectResult)
+				NodeResult = NodeToWidgetResult.FindRef(SearchResult.Node);
+				if (!NodeResult)
 				{
 					ensure(GraphResult);
-					ObjectResult = MakeShared<FFindInObjectTreeGraphResult>(GraphResult, CurSource, SearchResult.Object);
-					ObjectToWidgetResult.Add(SearchResult.Object, ObjectResult);
-					GraphResult->Children.Add(ObjectResult);
+					NodeResult = MakeShared<FFindInObjectTreeGraphResult>(GraphResult, SearchResult.Node);
+					NodeToWidgetResult.Add(SearchResult.Node, NodeResult);
+					GraphResult->Children.Add(NodeResult);
 				}
 			}
 
-			if (!SearchResult.PropertyName.IsNone())
+			if (!SearchResult.PinName.IsNone())
 			{
-				ensure(ObjectResult);
+				ensure(NodeResult);
 				FResultPtr PropertyResult = MakeShared<FFindInObjectTreeGraphResult>(
-						ObjectResult, CurSource, SearchResult.Object, SearchResult.PropertyName);
-				ObjectResult->Children.Add(PropertyResult);
+						NodeResult, SearchResult.Node, SearchResult.PinName);
+				NodeResult->Children.Add(PropertyResult);
 			}
 		}
 	}
-	
+
 	if (Results.IsEmpty())
 	{
 		Results.Add(MakeShared<FFindInObjectTreeGraphResult>(LOCTEXT("NoResults", "No results found")));
diff --git a/Engine/Plugins/Cameras/GameplayCameras/Source/GameplayCamerasEditor/Private/Editors/SObjectTreeGraphEditor.cpp b/Engine/Plugins/Cameras/GameplayCameras/Source/GameplayCamerasEditor/Private/Editors/SObjectTreeGraphEditor.cpp
--- a/Engine/Plugins/Cameras/GameplayCameras/Source/GameplayCamerasEditor/Private/Editors/SObjectTreeGraphEditor.cpp
+++ b/Engine/Plugins/Cameras/GameplayCameras/Source/GameplayCamerasEditor/Private/Editors/SObjectTreeGraphEditor.cpp
@@ -210,11 +210,32 @@ void SObjectTreeGraphEditor::InitializeBuiltInCommands()
 			);
 }
 
-void SObjectTreeGraphEditor::JumpToNode(UEdGraphNode* InNode)
+void SObjectTreeGraphEditor::JumpToNode(UEdGraphNode* InNode, FName InPinName)
 {
+	if (InPinName != NAME_None)
+	{
+		if (UEdGraphPin* Pin = InNode->FindPin(InPinName))
+		{
+			GraphEditor->JumpToPin(Pin);
+			return;
+		}
+	}
 	GraphEditor->JumpToNode(InNode);
 }
 
+bool SObjectTreeGraphEditor::FindAndJumpToObjectNode(UObject* InObject, FName InPropertyName)
+{
+	if (UObjectTreeGraph* CurrentGraph = Cast<UObjectTreeGraph>(GraphEditor->GetCurrentGraph()))
+	{
+		if (UObjectTreeGraphNode* Node = CurrentGraph->FindObjectNode(InObject))
+		{
+			JumpToNode(Node, InPropertyName);
+			return true;
+		}
+	}
+	return false;
+}
+
 void SObjectTreeGraphEditor::ResyncDetailsView()
 {
 	OnGraphSelectionChanged(GraphEditor->GetSelectedNodes());
diff --git a/Engine/Plugins/Cameras/GameplayCameras/Source/GameplayCamerasEditor/Private/Toolkits/CameraAssetEditorToolkit.cpp b/Engine/Plugins/Cameras/GameplayCameras/Source/GameplayCamerasEditor/Private/Toolkits/CameraAssetEditorToolkit.cpp
--- a/Engine/Plugins/Cameras/GameplayCameras/Source/GameplayCamerasEditor/Private/Toolkits/CameraAssetEditorToolkit.cpp
+++ b/Engine/Plugins/Cameras/GameplayCameras/Source/GameplayCamerasEditor/Private/Toolkits/CameraAssetEditorToolkit.cpp
@@ -153,8 +153,8 @@ void FCameraAssetEditorToolkit::CreateWidgets()
 
 	// Create the search panel.
 	SearchWidget = SNew(SFindInObjectTreeGraph)
-		.OnGetRootObjectsToSearch(this, &FCameraAssetEditorToolkit::OnGetRootObjectsToSearch)
-		.OnJumpToObjectRequested(this, &FCameraAssetEditorToolkit::OnJumpToObject);
+		.OnGetGraphsToSearch(this, &FCameraAssetEditorToolkit::OnGetGraphsToSearch)
+		.OnJumpToNodeRequested(this, &FCameraAssetEditorToolkit::OnJumpToNode);
 
 	// Create the message log.
 	BuildLogToolkit->Initialize("CameraAssetBuildMessages");
@@ -389,11 +389,20 @@ void FCameraAssetEditorToolkit::OnFindInCamera()
 	SearchWidget->FocusSearchEditBox();
 }
 
-void FCameraAssetEditorToolkit::OnGetRootObjectsToSearch(TArray<FFindInObjectTreeGraphSource>& OutSources)
+void FCameraAssetEditorToolkit::OnGetGraphsToSearch(TArray<FFindInObjectTreeGraphSource>& OutSources)
 {
 	TSharedPtr<FCameraSharedTransitionsAssetEditorMode> SharedTransitionsMode = GetTypedEditorMode<FCameraSharedTransitionsAssetEditorMode>(
 			FCameraSharedTransitionsAssetEditorMode::ModeName);
-	SharedTransitionsMode->OnGetRootObjectsToSearch(OutSources);
+	SharedTransitionsMode->OnGetGraphsToSearch(OutSources);
+}
+
+void FCameraAssetEditorToolkit::OnJumpToNode(UEdGraphNode* Node, FName PinName)
+{
+	// We only have one graph, the shared transition graph, so go there directly.
+	TSharedPtr<FCameraSharedTransitionsAssetEditorMode> SharedTransitionsMode = GetTypedEditorMode<FCameraSharedTransitionsAssetEditorMode>(
+			FCameraSharedTransitionsAssetEditorMode::ModeName);
+	SetEditorMode(FCameraSharedTransitionsAssetEditorMode::ModeName);
+	SharedTransitionsMode->JumpToNode(Node, PinName);
 }
 
 void FCameraAssetEditorToolkit::OnJumpToObject(UObject* Object)
diff --git a/Engine/Plugins/Cameras/GameplayCameras/Source/GameplayCamerasEditor/Private/Toolkits/CameraRigAssetEditorToolkit.cpp b/Engine/Plugins/Cameras/GameplayCameras/Source/GameplayCamerasEditor/Private/Toolkits/CameraRigAssetEditorToolkit.cpp
--- a/Engine/Plugins/Cameras/GameplayCameras/Source/GameplayCamerasEditor/Private/Toolkits/CameraRigAssetEditorToolkit.cpp
+++ b/Engine/Plugins/Cameras/GameplayCameras/Source/GameplayCamerasEditor/Private/Toolkits/CameraRigAssetEditorToolkit.cpp
@@ -217,8 +217,8 @@ void FCameraRigAssetEditorToolkit::CreateWidgets()
 	TArray<UEdGraph*> CameraRigGraphs;
 	CameraRigEditorWidget->GetGraphs(CameraRigGraphs);
 	SearchWidget = SNew(SFindInObjectTreeGraph)
-		.OnGetRootObjectsToSearch(this, &FCameraRigAssetEditorToolkit::OnGetRootObjectsToSearch)
-		.OnJumpToObjectRequested(this, &FCameraRigAssetEditorToolkit::OnJumpToObject);
+		.OnGetGraphsToSearch(this, &FCameraRigAssetEditorToolkit::OnGetGraphsToSearch)
+		.OnJumpToNodeRequested(this, &FCameraRigAssetEditorToolkit::OnJumpToNode);
 
 	// Create the message log.
 	BuildLogToolkit->Initialize("CameraRigAssetBuildMessages");
@@ -333,7 +333,8 @@ void FCameraRigAssetEditorToolkit::PostRegenerateMenusAndToolbars()
 
 void FCameraRigAssetEditorToolkit::OnCameraObjectInterfaceParameterSelected(UCameraObjectInterfaceParameterBase* Object)
 {
-	OnJumpToObject(Object, NAME_None);
+	TSharedPtr<SCameraRigAssetEditor> CameraRigEditor = Impl->GetCameraRigAssetEditor();
+	CameraRigEditor->FindAndJumpToObjectNode(Object);
 }
 
 void FCameraRigAssetEditorToolkit::OnBuild()
@@ -375,17 +376,21 @@ void FCameraRigAssetEditorToolkit::OnFindInCameraRig()
 	SearchWidget->FocusSearchEditBox();
 }
 
-void FCameraRigAssetEditorToolkit::OnGetRootObjectsToSearch(TArray<FFindInObjectTreeGraphSource>& OutSources)
+void FCameraRigAssetEditorToolkit::OnGetGraphsToSearch(TArray<FFindInObjectTreeGraphSource>& OutSources)
 {
-	UCameraRigAsset* CameraRig = Impl->GetCameraRigAsset();
-	OutSources.Add(FFindInObjectTreeGraphSource{ CameraRig, &NodeGraphConfig });
-	OutSources.Add(FFindInObjectTreeGraphSource{ CameraRig, &TransitionGraphConfig });
+	TArray<UEdGraph*> Graphs;
+	TSharedPtr<SCameraRigAssetEditor> CameraRigEditor = Impl->GetCameraRigAssetEditor();
+	CameraRigEditor->GetGraphs(Graphs);
+	for (UEdGraph* Graph : Graphs)
+	{
+		OutSources.Add(FFindInObjectTreeGraphSource{ Graph });
+	}
 }
 
-void FCameraRigAssetEditorToolkit::OnJumpToObject(UObject* Object, FName PropertyName)
+void FCameraRigAssetEditorToolkit::OnJumpToNode(UEdGraphNode* Node, FName PinName)
 {
 	TSharedPtr<SCameraRigAssetEditor> CameraRigEditor = Impl->GetCameraRigAssetEditor();
-	CameraRigEditor->FindAndJumpToObjectNode(Object);
+	CameraRigEditor->JumpToNode(Node, PinName);
 }
 
 void FCameraRigAssetEditorToolkit::OnInvokeCurveEditor(UObject* Object, FName PropertyName)
diff --git a/Engine/Plugins/Cameras/GameplayCameras/Source/GameplayCamerasEditor/Private/Toolkits/CameraShakeAssetEditorToolkit.cpp b/Engine/Plugins/Cameras/GameplayCameras/Source/GameplayCamerasEditor/Private/Toolkits/CameraShakeAssetEditorToolkit.cpp
--- a/Engine/Plugins/Cameras/GameplayCameras/Source/GameplayCamerasEditor/Private/Toolkits/CameraShakeAssetEditorToolkit.cpp
+++ b/Engine/Plugins/Cameras/GameplayCameras/Source/GameplayCamerasEditor/Private/Toolkits/CameraShakeAssetEditorToolkit.cpp
@@ -248,8 +248,8 @@ void FCameraShakeAssetEditorToolkit::CreateWidgets()
 
 	// Create the search panel.
 	SearchWidget = SNew(SFindInObjectTreeGraph)
-		.OnGetRootObjectsToSearch(this, &FCameraShakeAssetEditorToolkit::OnGetRootObjectsToSearch)
-		.OnJumpToObjectRequested(this, &FCameraShakeAssetEditorToolkit::OnJumpToObject);
+		.OnGetGraphsToSearch(this, &FCameraShakeAssetEditorToolkit::OnGetGraphsToSearch)
+		.OnJumpToNodeRequested(this, &FCameraShakeAssetEditorToolkit::OnJumpToNode);
 
 	// Create the message log.
 	BuildLogToolkit->Initialize("CameraShakeAssetBuildMessages");
@@ -402,22 +402,24 @@ void FCameraShakeAssetEditorToolkit::OnFindInCameraShake()
 	SearchWidget->FocusSearchEditBox();
 }
 
-void FCameraShakeAssetEditorToolkit::OnGetRootObjectsToSearch(TArray<FFindInObjectTreeGraphSource>& OutSources)
+void FCameraShakeAssetEditorToolkit::OnGetGraphsToSearch(TArray<FFindInObjectTreeGraphSource>& OutSources)
 {
-	OutSources.Add(FFindInObjectTreeGraphSource{ CameraShakeAsset, &NodeGraphConfig });
+	OutSources.Add(FFindInObjectTreeGraphSource{ NodeGraph });
 }
 
-void FCameraShakeAssetEditorToolkit::OnFocusHome()
+void FCameraShakeAssetEditorToolkit::OnJumpToNode(UEdGraphNode* Node, FName PinName)
 {
-	OnJumpToObject(CameraShakeAsset, NAME_None);
+	NodeGraphEditor->JumpToNode(Node, PinName);
 }
 
 void FCameraShakeAssetEditorToolkit::OnJumpToObject(UObject* Object, FName PropertyName)
 {
-	if (UObjectTreeGraphNode* NodeGraphObjectNode = NodeGraph->FindObjectNode(Object))
-	{
-		NodeGraphEditor->JumpToNode(NodeGraphObjectNode);
-	}
+	NodeGraphEditor->FindAndJumpToObjectNode(Object, PropertyName);
+}
+
+void FCameraShakeAssetEditorToolkit::OnFocusHome()
+{
+	NodeGraphEditor->FindAndJumpToObjectNode(CameraShakeAsset);
 }
 
 FText FCameraShakeAssetEditorToolkit::GetBaseToolkitName() const
diff --git a/Engine/Plugins/Cameras/GameplayCameras/Source/GameplayCamerasEditor/Private/Toolkits/CameraSharedTransitionsAssetEditorMode.cpp b/Engine/Plugins/Cameras/GameplayCameras/Source/GameplayCamerasEditor/Private/Toolkits/CameraSharedTransitionsAssetEditorMode.cpp
--- a/Engine/Plugins/Cameras/GameplayCameras/Source/GameplayCamerasEditor/Private/Toolkits/CameraSharedTransitionsAssetEditorMode.cpp
+++ b/Engine/Plugins/Cameras/GameplayCameras/Source/GameplayCamerasEditor/Private/Toolkits/CameraSharedTransitionsAssetEditorMode.cpp
@@ -80,12 +80,21 @@ void FCameraSharedTransitionsAssetEditorMode::OnDeactivateMode(const FAssetEdito
 	UToolMenus::UnregisterOwner(this);
 }
 
-void FCameraSharedTransitionsAssetEditorMode::OnGetRootObjectsToSearch(TArray<FFindInObjectTreeGraphSource>& OutSources)
+void FCameraSharedTransitionsAssetEditorMode::OnGetGraphsToSearch(TArray<FFindInObjectTreeGraphSource>& OutSources)
 {
-	OutSources.Add(FFindInObjectTreeGraphSource{ CameraAsset, &TransitionGraphConfig });
+	UEdGraph* TransitionGraph = Impl->GetCameraRigTransitionEditor()->GetTransitionGraph();
+	OutSources.Add(FFindInObjectTreeGraphSource{ TransitionGraph });
 }
 
-bool FCameraSharedTransitionsAssetEditorMode::JumpToObject(UObject* InObject, FName PropertyName)
+void FCameraSharedTransitionsAssetEditorMode::JumpToNode(UEdGraphNode* InNode, FName InPinName)
+{
+	if (TSharedPtr<SCameraRigTransitionEditor> TransitionEditor = Impl->GetCameraRigTransitionEditor())
+	{
+		TransitionEditor->JumpToNode(InNode, InPinName);
+	}
+}
+
+bool FCameraSharedTransitionsAssetEditorMode::JumpToObject(UObject* InObject, FName InPropertyName)
 {
 	if (TSharedPtr<SCameraRigTransitionEditor> TransitionEditor = Impl->GetCameraRigTransitionEditor())
 	{
diff --git a/Engine/Plugins/Cameras/GameplayCameras/Source/GameplayCamerasEditor/Public/Editors/ObjectTreeGraphSearch.h b/Engine/Plugins/Cameras/GameplayCameras/Source/GameplayCamerasEditor/Public/Editors/ObjectTreeGraphSearch.h
--- a/Engine/Plugins/Cameras/GameplayCameras/Source/GameplayCamerasEditor/Public/Editors/ObjectTreeGraphSearch.h
+++ b/Engine/Plugins/Cameras/GameplayCameras/Source/GameplayCamerasEditor/Public/Editors/ObjectTreeGraphSearch.h
@@ -3,30 +3,34 @@
 #pragma once
 
 #include "Containers/UnrealString.h"
+#include "Templates/SharedPointer.h"
 #include "UObject/WeakObjectPtr.h"
 
+class UEdGraph;
+class UEdGraphNode;
 class UObject;
 struct FObjectTreeGraphConfig;
 
 /** Search result for an object tree graph search. */
 struct FObjectTreeGraphSearchResult
 {
-	/** The root object that the current search result was found in. */
-	UObject* RootObject = nullptr;
+	/** The graph in which the result was found. */
+	UEdGraph* Graph = nullptr;
 
-	/** The graph config for the object hierarchy this result was found in. */
-	const FObjectTreeGraphConfig* GraphConfig = nullptr;
+	/** The node found to match the search. */
+	UEdGraphNode* Node = nullptr;
 
-	/** The object found to match the search. */
-	UObject* Object = nullptr;
+	/** The name of the pin that matched the search. None if the node itself matches. */
+	FName PinName;
 
-	/** 
-	 * The specific object property that matched the search. 
-	 * None if the object itself matches.
-	 */
-	FName PropertyName;
-};
+public:
 
+	FObjectTreeGraphSearchResult()
+	{}
+	FObjectTreeGraphSearchResult(UEdGraph* InGraph, UEdGraphNode* InNode, FName InPinName = NAME_None)
+		: Graph(InGraph), Node(InNode), PinName(InPinName)
+	{}
+};
 
 /**
  * A utility class that can search a series of string tokens across an
@@ -40,8 +44,10 @@ public:
 
 public:
 
-	/** Adds a root object to search through. */
-	void AddRootObject(UObject* InObject, const FObjectTreeGraphConfig* InGraphConfig);
+	/** Adds a graph to search inside of. */
+	void AddGraph(UEdGraph* InGraph);
+	/** Sets the graphs to search inside of. */
+	void SetGraphs(TConstArrayView<UEdGraph*> InGraphs);
 
 	/** Searchs for the given string tokens. */
 	void Search(TArrayView<FString> InTokens, TArray<FObjectTreeGraphSearchResult>& OutResults) const;
@@ -52,31 +58,18 @@ private:
 
 	struct FSearchState
 	{
-		const FObjectTreeGraphConfig* GraphConfig = nullptr;
+		UEdGraph* Graph = nullptr;
 		TArrayView<FString> Tokens;
-
-		UObject* RootObject = nullptr;
-		TArray<UObject*> ObjectStack;
-		TSet<UObject*> VisitedObjects;
-
 		TArray<FSearchResult> Results;
 	};
 
-	struct FRootObjectInfo
-	{
-		TWeakObjectPtr<> WeakRootObject;
-		const FObjectTreeGraphConfig* GraphConfig = nullptr;
-	};
-
-	void SearchRootObject(const FRootObjectInfo& InRootObjectInfo, TArrayView<FString> InTokens, TArray<FSearchResult>& OutResults) const;
-	void SearchObject(UObject* InObject, FSearchState& InOutState) const;
+	void SearchGraph(UEdGraph* Graph, FSearchState& State) const;
+	void SearchNode(UEdGraphNode* Node, FSearchState& State) const;
 
-	bool MatchObject(UObject* InObject, const FSearchState& InState) const;
-	bool MatchObjectProperty(UObject* InObject, FProperty* InProperty, const FSearchState& InState) const;
 	bool MatchString(const FString& InString, const FSearchState& InState) const;
 
 private:
 
-	TArray<FRootObjectInfo> RootObjectInfos;
+	TArray<UEdGraph*> SearchGraphs;
 };
 
diff --git a/Engine/Plugins/Cameras/GameplayCameras/Source/GameplayCamerasEditor/Public/Editors/SCameraRigAssetEditor.h b/Engine/Plugins/Cameras/GameplayCameras/Source/GameplayCamerasEditor/Public/Editors/SCameraRigAssetEditor.h
--- a/Engine/Plugins/Cameras/GameplayCameras/Source/GameplayCamerasEditor/Public/Editors/SCameraRigAssetEditor.h
+++ b/Engine/Plugins/Cameras/GameplayCameras/Source/GameplayCamerasEditor/Public/Editors/SCameraRigAssetEditor.h
@@ -76,6 +76,8 @@ public:
 
 	/** Focuses the current graph to the root object node. */
 	void FocusHome();
+	/** Jumps to a node. */
+	bool JumpToNode(UEdGraphNode* InNode, FName InPinName = NAME_None);
 	/** Finds a node for the given object and, if so, jumps to it. */
 	bool FindAndJumpToObjectNode(UObject* InObject);
 
diff --git a/Engine/Plugins/Cameras/GameplayCameras/Source/GameplayCamerasEditor/Public/Editors/SCameraRigTransitionEditor.h b/Engine/Plugins/Cameras/GameplayCameras/Source/GameplayCamerasEditor/Public/Editors/SCameraRigTransitionEditor.h
--- a/Engine/Plugins/Cameras/GameplayCameras/Source/GameplayCamerasEditor/Public/Editors/SCameraRigTransitionEditor.h
+++ b/Engine/Plugins/Cameras/GameplayCameras/Source/GameplayCamerasEditor/Public/Editors/SCameraRigTransitionEditor.h
@@ -61,7 +61,7 @@ public:
 	/** Focuses the current graph to the root object node. */
 	void FocusHome();
 	/** Jumps the current graph to the given node. */
-	void JumpToNode(UEdGraphNode* InGraphNode);
+	void JumpToNode(UEdGraphNode* InGraphNode, FName InPinName = NAME_None);
 	/** Finds a node for the given object and, if so, jumps to it. */
 	bool FindAndJumpToObjectNode(UObject* InObject);
 
diff --git a/Engine/Plugins/Cameras/GameplayCameras/Source/GameplayCamerasEditor/Public/Editors/SFindInObjectTreeGraph.h b/Engine/Plugins/Cameras/GameplayCameras/Source/GameplayCamerasEditor/Public/Editors/SFindInObjectTreeGraph.h
--- a/Engine/Plugins/Cameras/GameplayCameras/Source/GameplayCamerasEditor/Public/Editors/SFindInObjectTreeGraph.h
+++ b/Engine/Plugins/Cameras/GameplayCameras/Source/GameplayCamerasEditor/Public/Editors/SFindInObjectTreeGraph.h
@@ -2,12 +2,9 @@
 
 #pragma once
 
-#include "CoreTypes.h"
-#include "EdGraph/EdGraphPin.h"
 #include "UObject/WeakObjectPtr.h"
 #include "Templates/SharedPointerFwd.h"
 #include "Widgets/SCompoundWidget.h"
-#include "Widgets/Views/STableRow.h"
 #include "Widgets/Views/STableViewBase.h"
 #include "Widgets/Views/STreeView.h"
 
@@ -16,6 +13,7 @@ class SSearchBox;
 class UEdGraph;
 class UEdGraphNode;
 class UEdGraphSchema;
+class UObjectTreeGraph;
 struct FObjectTreeGraphConfig;
 
 /**
@@ -23,10 +21,8 @@ struct FObjectTreeGraphConfig;
  */
 struct FFindInObjectTreeGraphSource
 {
-	/** The root object a result was found in. */
-	UObject* RootObject = nullptr;
-	/** The config for the object tree graph a result was found in. */
-	const FObjectTreeGraphConfig* GraphConfig = nullptr;
+	/** The graph where results are to be found. */
+	UEdGraph* Graph = nullptr;
 };
 
 /** 
@@ -43,10 +39,10 @@ public:
 
 	/** Custom text for this result. */
 	FText CustomText;
-	/** The object that this result refers to. */
-	TWeakObjectPtr<> WeakObject;
-	/** The property name that this result refers to. */
-	FName PropertyName;
+	/** The node that this result refers to. */
+	TWeakObjectPtr<UEdGraphNode> WeakNode;
+	/** The pin name that this result refers to. */
+	FName PinName;
 
 public:
 
@@ -55,14 +51,8 @@ public:
 	/** Creates a new result referring to an object, under a parent result. */
 	FFindInObjectTreeGraphResult(
 			TSharedPtr<FFindInObjectTreeGraphResult>& InParent, 
-			const FFindInObjectTreeGraphSource& InSource, 
-			UObject* InObject);
-	/** Creates a new result referring to an object's property, under a parent result. */
-	FFindInObjectTreeGraphResult(
-			TSharedPtr<FFindInObjectTreeGraphResult>& InParent, 
-			const FFindInObjectTreeGraphSource& InSource, 
-			UObject* InObject, 
-			FName InPropertyName);
+			UEdGraphNode* InNode,
+			FName PinName = NAME_None);
 
 	/** Gets the icon for this result. */
 	TSharedRef<SWidget>	GetIcon() const;
@@ -75,10 +65,6 @@ public:
 
 	/** Go to the graph node, pin, etc. */
 	FReply OnClick(TSharedRef<SFindInObjectTreeGraph> FindInObjectTreeGraph);
-
-private:
-
-	FFindInObjectTreeGraphSource Source;
 };
 
 /**
@@ -88,21 +74,23 @@ class SFindInObjectTreeGraph : public SCompoundWidget
 {
 public:
 
-	DECLARE_DELEGATE_OneParam(FOnGetRootObjectsToSearch, TArray<FFindInObjectTreeGraphSource>&);
-	DECLARE_DELEGATE_TwoParams(FOnJumpToObjectRequested, UObject*, FName);
+	DECLARE_DELEGATE_OneParam(FOnGetGraphsToSearch, TArray<FFindInObjectTreeGraphSource>&);
+	DECLARE_DELEGATE_TwoParams(FOnJumpToNodeRequested, UEdGraphNode*, FName);
 
 	SLATE_BEGIN_ARGS(SFindInObjectTreeGraph)
 	{}
 		/** The callback to get the graphs to search. */
-		SLATE_EVENT(FOnGetRootObjectsToSearch, OnGetRootObjectsToSearch)
-		/** The callback to invoke when a search result wants to focus an object node or one of its pins. */
-		SLATE_EVENT(FOnJumpToObjectRequested, OnJumpToObjectRequested)
+		SLATE_EVENT(FOnGetGraphsToSearch, OnGetGraphsToSearch)
+		/** The callback to invoke when a search result wants to focus a node or one of its pins. */
+		SLATE_EVENT(FOnJumpToNodeRequested, OnJumpToNodeRequested)
 	SLATE_END_ARGS()
 
 	void Construct(const FArguments& InArgs);
 
 	void FocusSearchEditBox();
 
+	void Search(const FString& InSearchQuery);
+
 protected:
 
 	typedef TSharedPtr<FFindInObjectTreeGraphResult> FResultPtr;
@@ -122,8 +110,8 @@ protected:
 
 protected:
 
-	FOnGetRootObjectsToSearch OnGetRootObjectsToSearch;
-	FOnJumpToObjectRequested OnJumpToObjectRequested;
+	FOnGetGraphsToSearch OnGetGraphsToSearch;
+	FOnJumpToNodeRequested OnJumpToNodeRequested;
 
 	TSharedPtr<SSearchBox> SearchBox;
 	TSharedPtr<SResultTreeView> ResultTreeView;
diff --git a/Engine/Plugins/Cameras/GameplayCameras/Source/GameplayCamerasEditor/Public/Editors/SObjectTreeGraphEditor.h b/Engine/Plugins/Cameras/GameplayCameras/Source/GameplayCamerasEditor/Public/Editors/SObjectTreeGraphEditor.h
--- a/Engine/Plugins/Cameras/GameplayCameras/Source/GameplayCamerasEditor/Public/Editors/SObjectTreeGraphEditor.h
+++ b/Engine/Plugins/Cameras/GameplayCameras/Source/GameplayCamerasEditor/Public/Editors/SObjectTreeGraphEditor.h
@@ -54,7 +54,9 @@ public:
 	void RegisterEditor();
 	void UnregisterEditor();
 
-	void JumpToNode(UEdGraphNode* InNode);
+	void JumpToNode(UEdGraphNode* InNode, FName InPinName = NAME_None);
+	bool FindAndJumpToObjectNode(UObject* InObject, FName InPropertyName = NAME_None);
+
 	void ResyncDetailsView();
 
 	TSharedPtr<SGraphEditor> GetGraphEditor() { return GraphEditor; }
diff --git a/Engine/Plugins/Cameras/GameplayCameras/Source/GameplayCamerasEditor/Public/Toolkits/CameraAssetEditorToolkit.h b/Engine/Plugins/Cameras/GameplayCameras/Source/GameplayCamerasEditor/Public/Toolkits/CameraAssetEditorToolkit.h
--- a/Engine/Plugins/Cameras/GameplayCameras/Source/GameplayCamerasEditor/Public/Toolkits/CameraAssetEditorToolkit.h
+++ b/Engine/Plugins/Cameras/GameplayCameras/Source/GameplayCamerasEditor/Public/Toolkits/CameraAssetEditorToolkit.h
@@ -78,7 +78,9 @@ private:
 	void OnBuild();
 	void OnFindInCamera();
 
-	void OnGetRootObjectsToSearch(TArray<FFindInObjectTreeGraphSource>& OutSources);
+	void OnGetGraphsToSearch(TArray<FFindInObjectTreeGraphSource>& OutSources);
+	void OnJumpToNode(UEdGraphNode* Node, FName PinName = NAME_None);
+
 	void OnJumpToObject(UObject* Object);
 	void OnJumpToObject(UObject* Object, FName PropertyName);
 
diff --git a/Engine/Plugins/Cameras/GameplayCameras/Source/GameplayCamerasEditor/Public/Toolkits/CameraRigAssetEditorToolkit.h b/Engine/Plugins/Cameras/GameplayCameras/Source/GameplayCamerasEditor/Public/Toolkits/CameraRigAssetEditorToolkit.h
--- a/Engine/Plugins/Cameras/GameplayCameras/Source/GameplayCamerasEditor/Public/Toolkits/CameraRigAssetEditorToolkit.h
+++ b/Engine/Plugins/Cameras/GameplayCameras/Source/GameplayCamerasEditor/Public/Toolkits/CameraRigAssetEditorToolkit.h
@@ -70,8 +70,8 @@ private:
 	void OnBuild();
 	void OnFindInCameraRig();
 
-	void OnGetRootObjectsToSearch(TArray<FFindInObjectTreeGraphSource>& OutSources);
-	void OnJumpToObject(UObject* Object, FName PropertyName);
+	void OnGetGraphsToSearch(TArray<FFindInObjectTreeGraphSource>& OutSources);
+	void OnJumpToNode(UEdGraphNode* Node, FName PinName);
 
 	void OnInvokeCurveEditor(UObject* Object, FName PropertyName);
 
diff --git a/Engine/Plugins/Cameras/GameplayCameras/Source/GameplayCamerasEditor/Public/Toolkits/CameraShakeAssetEditorToolkit.h b/Engine/Plugins/Cameras/GameplayCameras/Source/GameplayCamerasEditor/Public/Toolkits/CameraShakeAssetEditorToolkit.h
--- a/Engine/Plugins/Cameras/GameplayCameras/Source/GameplayCamerasEditor/Public/Toolkits/CameraShakeAssetEditorToolkit.h
+++ b/Engine/Plugins/Cameras/GameplayCameras/Source/GameplayCamerasEditor/Public/Toolkits/CameraShakeAssetEditorToolkit.h
@@ -83,9 +83,11 @@ private:
 	void OnBuild();
 	void OnFindInCameraShake();
 
-	void OnGetRootObjectsToSearch(TArray<FFindInObjectTreeGraphSource>& OutSources);
-	void OnFocusHome();
+	void OnGetGraphsToSearch(TArray<FFindInObjectTreeGraphSource>& OutSources);
+	void OnJumpToNode(UEdGraphNode* Node, FName PinName);
+
 	void OnJumpToObject(UObject* Object, FName PropertyName);
+	void OnFocusHome();
 
 private:
 
diff --git a/Engine/Plugins/Cameras/GameplayCameras/Source/GameplayCamerasEditor/Public/Toolkits/CameraSharedTransitionsAssetEditorMode.h b/Engine/Plugins/Cameras/GameplayCameras/Source/GameplayCamerasEditor/Public/Toolkits/CameraSharedTransitionsAssetEditorMode.h
--- a/Engine/Plugins/Cameras/GameplayCameras/Source/GameplayCamerasEditor/Public/Toolkits/CameraSharedTransitionsAssetEditorMode.h
+++ b/Engine/Plugins/Cameras/GameplayCameras/Source/GameplayCamerasEditor/Public/Toolkits/CameraSharedTransitionsAssetEditorMode.h
@@ -24,8 +24,9 @@ public:
 
 	FCameraSharedTransitionsAssetEditorMode(UCameraAsset* InCameraAsset);
 
-	void OnGetRootObjectsToSearch(TArray<FFindInObjectTreeGraphSource>& OutSources);
-	bool JumpToObject(UObject* InObject, FName PropertyName);
+	void OnGetGraphsToSearch(TArray<FFindInObjectTreeGraphSource>& OutSources);
+	void JumpToNode(UEdGraphNode* InNode, FName InPinName);
+	bool JumpToObject(UObject* InObject, FName InPropertyName);
 
 protected:
 
