Cameras: continue recording debug info even when a camera system doesn't have any evaluation context active

#Synced-CL 48619321

diff --git a/Engine/Plugins/Cameras/GameplayCameras/Source/GameplayCameras/Private/Core/CameraSystemEvaluator.cpp b/Engine/Plugins/Cameras/GameplayCameras/Source/GameplayCameras/Private/Core/CameraSystemEvaluator.cpp
--- a/Engine/Plugins/Cameras/GameplayCameras/Source/GameplayCameras/Private/Core/CameraSystemEvaluator.cpp
+++ b/Engine/Plugins/Cameras/GameplayCameras/Source/GameplayCameras/Private/Core/CameraSystemEvaluator.cpp
@@ -247,6 +247,13 @@ void FCameraSystemEvaluator::UpdateImpl(float DeltaTime, ECameraNodeEvaluationTy
 	{
 		Result.bIsValid = false;
 		PreVisualResult.bIsValid = false;
+
+#if UE_GAMEPLAY_CAMERAS_DEBUG
+		BuildDebugBlocksIfNeeded();
+#endif  // UE_GAMEPLAY_CAMERAS_DEBUG
+
+		ContextStack.OnEndCameraSystemUpdate();
+
 		return;
 	}
 
