Cameras: make separate "getter" nodes for camera parameters, describe their connections in more general way

- Camera parameters were previously objects that doubled as their own node. This prevented driving two or more camera properties with the same parameter. This change removes the "node" aspects from camera parameters and adds a new "getter" node for that. The "getter" nodes point bac to the parameter they represent.
- Previous data is gracefully upgraded for camera parameters that were present as nodes.
- Since parameters may be connected to multiple pins, we can't determine their default value from that anymore, so the parameter has its own value now, a bit like variables in Blueprint. This changes a few things regarding the build process. A new customization class allows showing that default value in Details View for the user to set. Other parameter options (like "pre-blended") are removed from the list-view in the parameters panel: parameters are now meant to be tweaked in the Details View.
- Parameters can be added in transition graphs, for driving properties on blend nodes and such. A new base schema class provides support for this (a lot of its code is moved from the the camera node tree graph schema).
- Actions for adding or auto-wiring parameters are changed accordingly (they now deal with the "getter" node mainly).

#Synced-CL 48619321

diff --git a/Engine/Plugins/Cameras/GameplayCameras/Source/GameplayCameras/Private/Build/CameraAssetBuilder.cpp b/Engine/Plugins/Cameras/GameplayCameras/Source/GameplayCameras/Private/Build/CameraAssetBuilder.cpp
--- a/Engine/Plugins/Cameras/GameplayCameras/Source/GameplayCameras/Private/Build/CameraAssetBuilder.cpp
+++ b/Engine/Plugins/Cameras/GameplayCameras/Source/GameplayCameras/Private/Build/CameraAssetBuilder.cpp
@@ -164,7 +164,7 @@ void FCameraAssetBuilder::BuildCameraImpl(bool bBuildReferencedAssets)
 
 	for (const UCameraRigAsset* CameraRig : AllCameraRigs)
 	{
-		FCameraObjectInterfaceParameterBuilder::SetDefaultParameterValues(CameraRig, DefaultParameters);
+		CopyDefaultParameterValues(CameraRig, DefaultParameters);
 	}
 
 	if (!DefaultParameters.Identical(&CameraAsset->DefaultParameters, 0))
@@ -189,6 +189,39 @@ void FCameraAssetBuilder::BuildCameraImpl(bool bBuildReferencedAssets)
 	}
 }
 
+void FCameraAssetBuilder::CopyDefaultParameterValues(const UCameraRigAsset* InCameraRig, FInstancedPropertyBag& DefaultParameters)
+{
+	// Our default parameters property bag is an aggregation of various camera rigs' parameters.
+	// Copy their default values over.
+	const FInstancedPropertyBag& SourceValues = InCameraRig->GetDefaultParameters();
+	const void* SourceValuesPtr = SourceValues.GetValue().GetMemory();
+	void* TargetValuesPtr = DefaultParameters.GetMutableValue().GetMemory();
+	const UPropertyBag* SourceValueStruct = SourceValues.GetPropertyBagStruct();
+	const UPropertyBag* TargetValueStruct = DefaultParameters.GetPropertyBagStruct();
+	if (!SourceValuesPtr || !TargetValuesPtr || !SourceValueStruct || !TargetValueStruct)
+	{
+		// We may have empty property bags for camera rigs with no parameters.
+		return;
+	}
+
+	for (const FPropertyBagPropertyDesc& SourcePropertyDesc : SourceValueStruct->GetPropertyDescs())
+	{
+		const FPropertyBagPropertyDesc* TargetPropertyDesc = TargetValueStruct->FindPropertyDescByID(SourcePropertyDesc.ID);
+		if (!ensure(TargetPropertyDesc && TargetPropertyDesc->CachedProperty && SourcePropertyDesc.CachedProperty))
+		{
+			continue;
+		}
+		if (!ensure(SourcePropertyDesc.CompatibleType(*TargetPropertyDesc)))
+		{
+			continue;
+		}
+
+		const void* SourcePtr = SourcePropertyDesc.CachedProperty->ContainerPtrToValuePtr<void>(SourceValuesPtr);
+		void* TargetPtr = TargetPropertyDesc->CachedProperty->ContainerPtrToValuePtr<void>(TargetValuesPtr);
+		SourcePropertyDesc.CachedProperty->CopyCompleteValue(TargetPtr, SourcePtr);
+	}
+}
+
 void FCameraAssetBuilder::UpdateBuildStatus()
 {
 	ECameraBuildStatus BuildStatus = ECameraBuildStatus::Clean;
diff --git a/Engine/Plugins/Cameras/GameplayCameras/Source/GameplayCameras/Private/Build/CameraObjectInterfaceBuilder.cpp b/Engine/Plugins/Cameras/GameplayCameras/Source/GameplayCameras/Private/Build/CameraObjectInterfaceBuilder.cpp
--- a/Engine/Plugins/Cameras/GameplayCameras/Source/GameplayCameras/Private/Build/CameraObjectInterfaceBuilder.cpp
+++ b/Engine/Plugins/Cameras/GameplayCameras/Source/GameplayCameras/Private/Build/CameraObjectInterfaceBuilder.cpp
@@ -2,6 +2,7 @@
 
 #include "Build/CameraObjectInterfaceBuilder.h"
 
+#include "Build/CameraBuildLog.h"
 #include "Core/BaseCameraObject.h"
 #include "Core/CameraNode.h"
 #include "Core/CameraNodeHierarchy.h"
@@ -78,6 +79,7 @@ struct FInterfaceParameterBindingBuilder
 	template<typename CameraParameterOrVariableReferenceType>
 	void SetCameraParameterOrVariableReferenceOverride(
 			const UCameraObjectInterfaceBlendableParameter* BlendableParameter,
+			UCameraNode* TargetNode,
 			FStructProperty* TargetProperty,
 			CameraParameterOrVariableReferenceType* CameraParameterOrVariableReferencePtr)
 	{
@@ -85,15 +87,13 @@ struct FInterfaceParameterBindingBuilder
 		using ValueType = typename VariableAssetType::ValueType;
 
 		ensure(BlendableParameter->PrivateVariableID.IsValid());
-		ensure(BlendableParameter->TargetPropertyName == TargetProperty->GetFName());
 
-		const bool bIsValid = CheckIfParameterCanBeOverridden(BlendableParameter, CameraParameterOrVariableReferencePtr);
+		const bool bIsValid = CheckIfParameterCanBeOverridden(BlendableParameter, TargetNode, TargetProperty->GetFName(), CameraParameterOrVariableReferencePtr);
 		if (!bIsValid)
 		{
 			return;
 		}
 
-		UObject* TargetNode = BlendableParameter->Target;
 		FCameraVariableID PreviousVariableID = FindOldDrivingVariableID(TargetProperty->GetFName(), TargetNode);
 		if (PreviousVariableID != BlendableParameter->PrivateVariableID)
 		{
@@ -106,20 +106,21 @@ struct FInterfaceParameterBindingBuilder
 	template<typename VariableAssetType>
 	void SetCustomBlendableParameterOverride(
 			const UCameraObjectInterfaceBlendableParameter* BlendableParameter,
+			UCameraNode* TargetNode,
+			FName TargetPropertyName,
 			const FCustomCameraNodeParameterInfos::FBlendableParameterInfo& CustomParameter)
 	{
 		using ValueType = typename VariableAssetType::ValueType;
 
 		ensure(BlendableParameter->PrivateVariableID.IsValid());
-		ensure(BlendableParameter->TargetPropertyName == CustomParameter.ParameterName);
+		ensure(TargetPropertyName == CustomParameter.ParameterName);
 
-		const bool bIsValid = CheckIfParameterCanBeOverridden(BlendableParameter, CustomParameter);
+		const bool bIsValid = CheckIfParameterCanBeOverridden(BlendableParameter, TargetNode, TargetPropertyName, CustomParameter);
 		if (!bIsValid)
 		{
 			return;
 		}
 
-		UObject* TargetNode = BlendableParameter->Target;
 		FCameraVariableID PreviousVariableID = FindOldDrivingVariableID(CustomParameter.ParameterName, TargetNode);
 		if (PreviousVariableID != BlendableParameter->PrivateVariableID)
 		{
@@ -131,12 +132,14 @@ struct FInterfaceParameterBindingBuilder
 
 	void SetCustomBlendableStructParameterOverride(
 			const UCameraObjectInterfaceBlendableParameter* BlendableParameter,
+			UCameraNode* TargetNode,
+			FName TargetPropertyName,
 			const FCustomCameraNodeParameterInfos::FBlendableParameterInfo& CustomParameter)
 	{
 		ensure(BlendableParameter->PrivateVariableID.IsValid());
-		ensure(BlendableParameter->TargetPropertyName == CustomParameter.ParameterName);
+		ensure(TargetPropertyName == CustomParameter.ParameterName);
 
-		const bool bIsValid = CheckIfParameterCanBeOverridden(BlendableParameter, CustomParameter);
+		const bool bIsValid = CheckIfParameterCanBeOverridden(BlendableParameter, TargetNode, TargetPropertyName, CustomParameter);
 		if (!bIsValid)
 		{
 			return;
@@ -145,14 +148,14 @@ struct FInterfaceParameterBindingBuilder
 		// Also ensure the struct type is compatible.
 		if (CustomParameter.BlendableStructType != BlendableParameter->BlendableStructType)
 		{
-			ReportError(BlendableParameter->Target,
+			ReportError(TargetNode,
 					FText::Format(
 						LOCTEXT(
 							"IncompatibleBlendableStructType", 
 							"Invalid interface parameter '{0}', driving property '{1}' on '{2}': expected type {3} but was {4}"),
 						FText::FromString(BlendableParameter->InterfaceParameterName),
-						FText::FromName(BlendableParameter->TargetPropertyName),
-						FText::FromName(BlendableParameter->Target->GetFName()),
+						FText::FromName(TargetPropertyName),
+						FText::FromName(TargetNode->GetFName()),
 #if WITH_EDITORONLY_DATA
 						CustomParameter.BlendableStructType->GetDisplayNameText(),
 						BlendableParameter->BlendableStructType->GetDisplayNameText()
@@ -164,7 +167,6 @@ struct FInterfaceParameterBindingBuilder
 			return;
 		}
 
-		UObject* TargetNode = BlendableParameter->Target;
 		FCameraVariableID PreviousVariableID = FindOldDrivingVariableID(CustomParameter.ParameterName, TargetNode);
 		if (PreviousVariableID != BlendableParameter->PrivateVariableID)
 		{
@@ -176,18 +178,18 @@ struct FInterfaceParameterBindingBuilder
 
 	void SetDataContextPropertyOverride(
 			const UCameraObjectInterfaceDataParameter* DataParameter,
+			UCameraNode* TargetNode,
 			FProperty* TargetProperty,
 			FCameraContextDataID* OverrideDataID)
 	{
 		ensure(DataParameter->PrivateDataID);
 
-		const bool bIsValid = CheckIfParameterCanBeOverridden(DataParameter, TargetProperty, OverrideDataID);
+		const bool bIsValid = CheckIfParameterCanBeOverridden(DataParameter, TargetNode, TargetProperty, OverrideDataID);
 		if (!bIsValid)
 		{
 			return;
 		}
 
-		UObject* TargetNode = DataParameter->Target;
 		FCameraContextDataID PreviousDataID = FindOldDrivingDataID(TargetProperty->GetFName(), TargetNode);
 		if (PreviousDataID != DataParameter->PrivateDataID)
 		{
@@ -199,18 +201,19 @@ struct FInterfaceParameterBindingBuilder
 
 	void SetCustomDataParameterOverride(
 			const UCameraObjectInterfaceDataParameter* DataParameter,
+			UCameraNode* TargetNode,
+			FName TargetPropertyName,
 			const FCustomCameraNodeParameterInfos::FDataParameterInfo& CustomParameter)
 	{
 		ensure(DataParameter->PrivateDataID);
-		ensure(DataParameter->TargetPropertyName == CustomParameter.ParameterName);
+		ensure(TargetPropertyName == CustomParameter.ParameterName);
 
-		const bool bIsValid = CheckIfParameterCanBeOverridden(DataParameter, CustomParameter);
+		const bool bIsValid = CheckIfParameterCanBeOverridden(DataParameter, TargetNode, TargetPropertyName, CustomParameter);
 		if (!bIsValid)
 		{
 			return;
 		}
 
-		UObject* TargetNode = DataParameter->Target;
 		FCameraContextDataID PreviousDataID = FindOldDrivingDataID(CustomParameter.ParameterName, TargetNode);
 		if (PreviousDataID != DataParameter->PrivateDataID)
 		{
@@ -245,17 +248,19 @@ private:
 	template<typename CameraParameterOrVariableReferenceType>
 	bool CheckIfParameterCanBeOverridden(
 			const UCameraObjectInterfaceBlendableParameter* BlendableParameter, 
+			UCameraNode* TargetNode,
+			const FName TargetPropertyName,
 			CameraParameterOrVariableReferenceType* CameraParameterOrVariableReference)
 	{
 		if (!MatchesVariableType<CameraParameterOrVariableReferenceType>(BlendableParameter->ParameterType))
 		{
-			ReportError(BlendableParameter->Target,
+			ReportError(TargetNode,
 					FText::Format(
 						LOCTEXT(
 							"BlendableParameterTypeMismatch",
 							"Camera node parameter '{0}.{1}' of type {2} is not compatible with camera rig parameter '{3}' of type {4}"),
-						FText::FromName(BlendableParameter->Target->GetFName()), 
-						FText::FromName(BlendableParameter->TargetPropertyName),
+						FText::FromName(TargetNode->GetFName()), 
+						FText::FromName(TargetPropertyName),
 						GetVariableTypeAsText<CameraParameterOrVariableReferenceType>(),
 						FText::FromString(BlendableParameter->InterfaceParameterName),
 						UEnum::GetDisplayValueAsText(BlendableParameter->ParameterType)));
@@ -263,31 +268,33 @@ private:
 		}
 		if (CameraParameterOrVariableReference->Variable != nullptr)
 		{
-			ReportError(BlendableParameter->Target,
+			ReportError(TargetNode,
 					FText::Format(
 						LOCTEXT(
 							"BlendableParameterDrivenTwice", 
 							"Camera node parameter '{0}.{1}' is both exposed and driven by a variable!"),
-						FText::FromName(BlendableParameter->Target->GetFName()), 
-						FText::FromName(BlendableParameter->TargetPropertyName)));
+						FText::FromName(TargetNode->GetFName()), 
+						FText::FromName(TargetPropertyName)));
 			return false;
 		}
 
-		return CheckIfParameterCanBeOverridden(BlendableParameter, &CameraParameterOrVariableReference->VariableID);
+		return CheckIfParameterCanBeOverridden(BlendableParameter, TargetNode, TargetPropertyName, &CameraParameterOrVariableReference->VariableID);
 	}
 
 	bool CheckIfParameterCanBeOverridden(
 			const UCameraObjectInterfaceBlendableParameter* BlendableParameter, 
+			UCameraNode* TargetNode,
+			const FName TargetPropertyName,
 			const FCustomCameraNodeParameterInfos::FBlendableParameterInfo& CustomParameter)
 	{
 		if (BlendableParameter->ParameterType != CustomParameter.ParameterType)
 		{
-			ReportError(BlendableParameter->Target,
+			ReportError(TargetNode,
 					FText::Format(
 						LOCTEXT(
 							"BlendableParameterTypeMismatch",
 							"Camera node parameter '{0}.{1}' of type {2} is not compatible with camera rig parameter '{3}' of type {4}"),
-						FText::FromName(BlendableParameter->Target->GetFName()), 
+						FText::FromName(TargetNode->GetFName()), 
 						FText::FromName(CustomParameter.ParameterName),
 						UEnum::GetDisplayValueAsText(CustomParameter.ParameterType),
 						FText::FromString(BlendableParameter->InterfaceParameterName),
@@ -296,43 +303,45 @@ private:
 		}
 		if (CustomParameter.OverrideVariable != nullptr)
 		{
-			ReportError(BlendableParameter->Target,
+			ReportError(TargetNode,
 					FText::Format(
 						LOCTEXT(
 							"BlendableParameterDrivenTwice", 
 							"Camera node parameter '{0}.{1}' is both exposed and driven by a variable!"),
-						FText::FromName(BlendableParameter->Target->GetFName()), 
-						FText::FromName(BlendableParameter->TargetPropertyName)));
+						FText::FromName(TargetNode->GetFName()), 
+						FText::FromName(TargetPropertyName)));
 			return false;
 		}
 
-		return CheckIfParameterCanBeOverridden(BlendableParameter, CustomParameter.OverrideVariableID);
+		return CheckIfParameterCanBeOverridden(BlendableParameter, TargetNode, TargetPropertyName, CustomParameter.OverrideVariableID);
 	}
 
 	bool CheckIfParameterCanBeOverridden(
 			const UCameraObjectInterfaceBlendableParameter* BlendableParameter, 
+			UCameraNode* TargetNode,
+			const FName TargetPropertyName,
 			FCameraVariableID* VariableID)
 	{
 		if (!VariableID)
 		{
-			ReportError(BlendableParameter->Target,
+			ReportError(TargetNode,
 					FText::Format(
 						LOCTEXT(
 							"BlendableParameterMissingOverrideID", 
 							"Camera node parameter '{0}.{1}' cannot be overriden by a parameter"),
-						FText::FromName(BlendableParameter->Target->GetFName()), 
-						FText::FromName(BlendableParameter->TargetPropertyName)));
+						FText::FromName(TargetNode->GetFName()), 
+						FText::FromName(TargetPropertyName)));
 			return false;
 		}
 		if (VariableID->IsValid())
 		{
-			ReportError(BlendableParameter->Target,
+			ReportError(TargetNode,
 					FText::Format(
 						LOCTEXT(
 							"BlendableParameterDrivenTwice", 
 							"Camera node parameter '{0}.{1}' is both exposed and driven by a variable!"),
-						FText::FromName(BlendableParameter->Target->GetFName()), 
-						FText::FromName(BlendableParameter->TargetPropertyName)));
+						FText::FromName(TargetNode->GetFName()), 
+						FText::FromName(TargetPropertyName)));
 			return false;
 		}
 		return true;
@@ -340,16 +349,18 @@ private:
 
 	bool CheckIfParameterCanBeOverridden(
 			const UCameraObjectInterfaceDataParameter* DataParameter,
+			UCameraNode* TargetNode,
+			const FName TargetPropertyName,
 			const FCustomCameraNodeParameterInfos::FDataParameterInfo& CustomParameter)
 	{
 		if (DataParameter->DataType != CustomParameter.ParameterType)
 		{
-			ReportError(DataParameter->Target,
+			ReportError(TargetNode,
 					FText::Format(
 						LOCTEXT(
 							"DataParameterTypeMismatch",
 							"Camera node parameter '{0}.{1}' of type {2} is not compatible with camera rig parameter '{3}' of type {4}"),
-						FText::FromName(DataParameter->Target->GetFName()), 
+						FText::FromName(TargetNode->GetFName()), 
 						FText::FromName(CustomParameter.ParameterName),
 						UEnum::GetDisplayValueAsText(CustomParameter.ParameterType),
 						FText::FromString(DataParameter->InterfaceParameterName),
@@ -358,12 +369,12 @@ private:
 		}
 		if (DataParameter->DataContainerType != CustomParameter.ParameterContainerType)
 		{
-			ReportError(DataParameter->Target,
+			ReportError(TargetNode,
 					FText::Format(
 						LOCTEXT(
 							"DataParameterContainerTypeMismatch",
 							"Camera node parameter '{0}.{1}' has a different container type than camera rig parameter '{2}' ({3} vs {4})"),
-						FText::FromName(DataParameter->Target->GetFName()), 
+						FText::FromName(TargetNode->GetFName()), 
 						FText::FromName(CustomParameter.ParameterName),
 						FText::FromString(DataParameter->InterfaceParameterName),
 						UEnum::GetDisplayValueAsText(CustomParameter.ParameterContainerType),
@@ -371,11 +382,12 @@ private:
 			return false;
 		}
 
-		return CheckIfParameterCanBeOverridden(DataParameter, CustomParameter.OverrideDataID);
+		return CheckIfParameterCanBeOverridden(DataParameter, TargetNode, TargetPropertyName, CustomParameter.OverrideDataID);
 	}
 
 	bool CheckIfParameterCanBeOverridden(
 			const UCameraObjectInterfaceDataParameter* DataParameter,
+			UCameraNode* TargetNode,
 			const FProperty* TargetProperty,
 			FCameraContextDataID* DataID)
 	{
@@ -398,12 +410,12 @@ private:
 		}
 		if (!bDataContainerTypeMatches)
 		{
-			ReportError(DataParameter->Target,
+			ReportError(TargetNode,
 					FText::Format(
 						LOCTEXT(
 							"DataParameterContainerTypeMismatch",
 							"Camera node parameter '{0}.{1}' has a different container type than camera rig parameter '{2}' ({3} vs {4})"),
-						FText::FromName(DataParameter->Target->GetFName()), 
+						FText::FromName(TargetNode->GetFName()), 
 						FText::FromName(TargetProperty->GetFName()),
 						FText::FromString(DataParameter->InterfaceParameterName),
 						TargetProperty->GetClass()->GetDisplayNameText(),
@@ -450,12 +462,12 @@ private:
 		}
 		if (!bDataTypeMatches)
 		{
-			ReportError(DataParameter->Target,
+			ReportError(TargetNode,
 					FText::Format(
 						LOCTEXT(
 							"DataParameterTypeMismatch",
 							"Camera node parameter '{0}.{1}' of type {2} is not compatible with camera rig parameter '{3}' of type {4}"),
-						FText::FromName(DataParameter->Target->GetFName()), 
+						FText::FromName(TargetNode->GetFName()), 
 						FText::FromName(TargetProperty->GetFName()),
 						TargetProperty->GetClass()->GetDisplayNameText(),
 						FText::FromString(DataParameter->InterfaceParameterName),
@@ -463,33 +475,35 @@ private:
 			return false;
 		}
 
-		return CheckIfParameterCanBeOverridden(DataParameter, DataID);
+		return CheckIfParameterCanBeOverridden(DataParameter, TargetNode, TargetProperty->GetFName(), DataID);
 	}
 
 	bool CheckIfParameterCanBeOverridden(
 			const UCameraObjectInterfaceDataParameter* DataParameter,
+			UCameraNode* TargetNode,
+			const FName TargetPropertyName,
 			FCameraContextDataID* DataID)
 	{
 		if (!DataID)
 		{
-			ReportError(DataParameter->Target,
+			ReportError(TargetNode,
 					FText::Format(
 						LOCTEXT(
 							"DataParameterMissingOverrideID", 
 							"Camera node parameter '{0}.{1}' cannot be overriden by a parameter"),
-						FText::FromName(DataParameter->Target->GetFName()), 
-						FText::FromName(DataParameter->TargetPropertyName)));
+						FText::FromName(TargetNode->GetFName()), 
+						FText::FromName(TargetPropertyName)));
 			return false;
 		}
 		if (DataID->IsValid())
 		{
-			ReportError(DataParameter->Target,
+			ReportError(TargetNode,
 					FText::Format(
 						LOCTEXT(
 							"DataParameterDrivenTwice",
 							"Camera node parameter '{0}.{1}' is somehow overriden twice!"),
-						FText::FromName(DataParameter->Target->GetFName()), 
-						FText::FromName(DataParameter->TargetPropertyName)));
+						FText::FromName(TargetNode->GetFName()), 
+						FText::FromName(TargetPropertyName)));
 			return false;
 		}
 		return true;
@@ -764,127 +778,131 @@ void FCameraObjectInterfaceBuilder::BuildInterfaceParameterBindings()
 	// Each time we need to check for either a custom property (via ICustomCameraNodeParameterProvider),
 	// or a UObject property found with reflection.
 
-	using FBuiltDrivenParameter = TTuple<UObject*, FName>;
-	TSet<FBuiltDrivenParameter> BuiltDrivenParameters;
-
 	const FString CameraObjectName = CameraObject->GetName();
 	const FString CameraObjectPathName = CameraObject->GetPathName();
 
-	for (const UCameraObjectInterfaceBlendableParameter* BlendableParameter : CameraObject->Interface.BlendableParameters)
+	// Gather up all connected parameters, and what they're connected to.
+	using FCameraNodeProperty = TTuple<UCameraNode*, FName>;
+	TSet<FCameraNodeProperty> ConnectedCameraNodeProperties;
+	TMap<UCameraObjectInterfaceBlendableParameter*, TArray<FCameraNodeProperty>> BlendableParameterConnections;
+	TMap<UCameraObjectInterfaceDataParameter*, TArray<FCameraNodeProperty>> DataParameterConnections;
+	for (const FCameraObjectConnection& Connection : CameraObject->Connections.Connections)
 	{
-		// Basic validations.
-		if (!BlendableParameter->Target)
+		if (Connection.SourcePropertyName != NAME_None)
 		{
-			BuildLog.AddMessage(EMessageSeverity::Warning,
-					BlendableParameter,
-					LOCTEXT(
-						"InvalidBlendableParameterTarget",
-						"Invalid interface parameter: it has no target node."));
 			continue;
 		}
-		if (BlendableParameter->TargetPropertyName.IsNone())
+		if (Connection.Target == nullptr || Connection.TargetPropertyName == NAME_None)
 		{
-			BuildLog.AddMessage(EMessageSeverity::Error,
-					BlendableParameter,
-					LOCTEXT(
-						"InvalidBlendableParameterTargetPropertyName", 
-						"Invalid interface parameter: it has not target property name."));
-			continue;
-		}
-		if (BlendableParameter->InterfaceParameterName.IsEmpty())
-		{
-			// Error already logged in BuildInterfaceParameters().
 			continue;
 		}
 
-		// Check duplicate bindings.
-		FBuiltDrivenParameter NewDrivenParameter(BlendableParameter->Target, BlendableParameter->TargetPropertyName);
-		if (BuiltDrivenParameters.Contains(NewDrivenParameter))
+		const UCameraObjectInterfaceParameterGetter* ParameterGetter = Cast<UCameraObjectInterfaceParameterGetter>(Connection.Source);
+		if (!ParameterGetter || !ParameterGetter->ParameterGuid.IsValid())
 		{
-			BuildLog.AddMessage(EMessageSeverity::Error,
-					FText::Format(LOCTEXT(
-						"BlendableParameterTargetCollision",
-						"Multiple interface parameters targeting property '{0}' on camera node '{1}'. Ignoring duplicates."),
-						FText::FromName(BlendableParameter->TargetPropertyName),
-						FText::FromName(BlendableParameter->Target->GetFName())));
 			continue;
 		}
-		BuiltDrivenParameters.Add(NewDrivenParameter);
 
-		// See if this interface parameter is overriding a camera node parameter.
-		// Otherwise, maybe it's targeting a camera rig node's override for an inner rig interface parameter.
-		if (SetupCustomBlendableParameterOverride(BlendableParameter))
-		{
-			// Implicit continue.
-		}
-		else if (SetupCameraParameterOrVariableReferenceOverride(BlendableParameter))
+		if (UCameraObjectInterfaceBlendableParameter* BlendableParameter = CameraObject->Interface.FindBlendableParameterByGuid(ParameterGetter->ParameterGuid))
 		{
-			// Implicit continue.
+			BlendableParameterConnections.FindOrAdd(BlendableParameter).Add(FCameraNodeProperty{ Connection.Target, Connection.TargetPropertyName });
 		}
-		else
+		else if (UCameraObjectInterfaceDataParameter* DataParameter = CameraObject->Interface.FindDataParameterByGuid(ParameterGetter->ParameterGuid))
 		{
-			UObject* Target = BlendableParameter->Target;
-			BuildLog.AddMessage(EMessageSeverity::Error,
-					Target,
-					FText::Format(LOCTEXT(
-						"InvalidBlendableParameterTargetProperty",
-						"Invalid interface parameter '{0}', driving property '{1}' on '{2}', but no such property found."),
-						FText::FromString(BlendableParameter->InterfaceParameterName), 
-						FText::FromName(BlendableParameter->TargetPropertyName),
-						FText::FromName(Target->GetFName())));
+			DataParameterConnections.FindOrAdd(DataParameter).Add(FCameraNodeProperty{ Connection.Target, Connection.TargetPropertyName });
 		}
 	}
 
-	for (const UCameraObjectInterfaceDataParameter* DataParameter : CameraObject->Interface.DataParameters)
+	auto CheckConnectionCollision = [this, &ConnectedCameraNodeProperties](const FCameraNodeProperty& InConnectedProperty) -> bool
 	{
-		// Basic validations.
-		if (!DataParameter->Target)
-		{
-			BuildLog.AddMessage(EMessageSeverity::Warning,
-					DataParameter,
-					LOCTEXT(
-						"InvalidDataParameterTarget",
-						"Invalid interface parameter: it has no target node."));
-			continue;
-		}
-		if (DataParameter->TargetPropertyName.IsNone())
+		bool bAlreadyConnected = false;
+		ConnectedCameraNodeProperties.Add(InConnectedProperty, &bAlreadyConnected);
+		if (bAlreadyConnected)
 		{
 			BuildLog.AddMessage(EMessageSeverity::Error,
-					DataParameter,
-					LOCTEXT(
-						"InvalidDataParameterTargetPropertyName", 
-						"Invalid interface parameter: it has not target property name."));
-			continue;
-		}
-		if (DataParameter->InterfaceParameterName.IsEmpty())
-		{
-			continue;
+					FText::Format(LOCTEXT(
+							"ParameterConnectionCollision",
+							"Multiple interface parameters targeting property '{0}' on camera node '{1}'. Ignoring duplicates."),
+						FText::FromName(InConnectedProperty.Value),
+						FText::FromName(InConnectedProperty.Key->GetFName())));
 		}
+		return bAlreadyConnected;
+	};
 
-		if (SetupCustomDataParameterOverride(DataParameter))
-		{
-			// Implicit continue.
-		}
-		else if (SetupDataContextPropertyOverride(DataParameter))
+	for (auto& Pair : BlendableParameterConnections)
+	{
+		const UCameraObjectInterfaceBlendableParameter* BlendableParameter(Pair.Key);
+		const TArray<FCameraNodeProperty>& ConnectedProperties(Pair.Value);
+	
+		for (const FCameraNodeProperty& ConnectedProperty : ConnectedProperties)
 		{
-			// Implicit continue.
+			// Check if some other parameter was already connected to this node's property.
+			if (CheckConnectionCollision(ConnectedProperty))
+			{
+				continue;
+			}
+
+			// See if this interface parameter is overriding a camera node parameter.
+			// Otherwise, maybe it's targeting a camera rig node's override for an inner rig interface parameter.
+			if (SetupCustomBlendableParameterOverride(BlendableParameter, ConnectedProperty.Key, ConnectedProperty.Value))
+			{
+				// Implicit continue.
+			}
+			else if (SetupCameraParameterOrVariableReferenceOverride(BlendableParameter, ConnectedProperty.Key, ConnectedProperty.Value))
+			{
+				// Implicit continue.
+			}
+			else
+			{
+				BuildLog.AddMessage(EMessageSeverity::Error,
+						ConnectedProperty.Key,
+						FText::Format(LOCTEXT(
+								"InvalidBlendableParameterTargetProperty",
+								"Invalid interface parameter '{0}', driving property '{1}' on '{2}', but no such property found."),
+							FText::FromString(BlendableParameter->InterfaceParameterName), 
+							FText::FromName(ConnectedProperty.Value),
+							FText::FromName(ConnectedProperty.Key->GetFName())));
+			}
 		}
-		else
+	}
+
+	for (auto& Pair : DataParameterConnections)
+	{
+		const UCameraObjectInterfaceDataParameter* DataParameter(Pair.Key);
+		const TArray<FCameraNodeProperty>& ConnectedProperties(Pair.Value);
+
+		for (const FCameraNodeProperty& ConnectedProperty : ConnectedProperties)
 		{
-			UObject* Target = DataParameter->Target;
-			BuildLog.AddMessage(EMessageSeverity::Error,
-					Target,
-					FText::Format(LOCTEXT(
-						"InvalidDataParameterTargetProperty",
-						"Invalid interface parameter '{0}', driving property '{1}' on '{2}', but no such property found."),
-						FText::FromString(DataParameter->InterfaceParameterName), 
-						FText::FromName(DataParameter->TargetPropertyName),
-						FText::FromName(Target->GetFName())));
+			// Check if some other parameter was already connected to this node's property.
+			if (CheckConnectionCollision(ConnectedProperty))
+			{
+				continue;
+			}
+
+			if (SetupCustomDataParameterOverride(DataParameter, ConnectedProperty.Key, ConnectedProperty.Value))
+			{
+				// Implicit continue.
+			}
+			else if (SetupDataContextPropertyOverride(DataParameter, ConnectedProperty.Key, ConnectedProperty.Value))
+			{
+				// Implicit continue.
+			}
+			else
+			{
+				BuildLog.AddMessage(EMessageSeverity::Error,
+						ConnectedProperty.Key,
+						FText::Format(LOCTEXT(
+								"InvalidDataParameterTargetProperty",
+								"Invalid interface parameter '{0}', driving property '{1}' on '{2}', but no such property found."),
+							FText::FromString(DataParameter->InterfaceParameterName), 
+							FText::FromName(ConnectedProperty.Value),
+							FText::FromName(ConnectedProperty.Key->GetFName())));
+			}
 		}
 	}
 }
 
-bool FCameraObjectInterfaceBuilder::SetupCameraParameterOrVariableReferenceOverride(const UCameraObjectInterfaceBlendableParameter* BlendableParameter)
+bool FCameraObjectInterfaceBuilder::SetupCameraParameterOrVariableReferenceOverride(const UCameraObjectInterfaceBlendableParameter* BlendableParameter, UCameraNode* Target, FName TargetPropertyName)
 {
 	using namespace Internal;
 
@@ -894,9 +912,8 @@ bool FCameraObjectInterfaceBuilder::SetupCameraParameterOrVariableReferenceOverr
 	// private variable of the given interface parameter, checking that the types match (UBooleanCameraVariable, 
 	// UInteger32CameraVariable, etc.)
 
-	UObject* Target = BlendableParameter->Target;
 	UClass* TargetClass = Target->GetClass();
-	FProperty* TargetProperty = TargetClass->FindPropertyByName(BlendableParameter->TargetPropertyName);
+	FProperty* TargetProperty = TargetClass->FindPropertyByName(TargetPropertyName);
 	if (!TargetProperty)
 	{
 		// No match, try something else.
@@ -912,7 +929,7 @@ bool FCameraObjectInterfaceBuilder::SetupCameraParameterOrVariableReferenceOverr
 						"InvalidCameraNodeParameter",
 						"Invalid interface parameter '{0}', driving property '{1}' on '{2}', but it's not a camera parameter."),
 					FText::FromString(BlendableParameter->InterfaceParameterName), 
-					FText::FromName(BlendableParameter->TargetPropertyName),
+					FText::FromName(TargetPropertyName),
 					FText::FromName(Target->GetFName())));
 		return true;
 	}
@@ -925,14 +942,14 @@ bool FCameraObjectInterfaceBuilder::SetupCameraParameterOrVariableReferenceOverr
 	{\
 		auto* CameraParameterPtr = TargetStructProperty->ContainerPtrToValuePtr<F##ValueName##CameraParameter>(Target);\
 		Builder.SetCameraParameterOrVariableReferenceOverride<F##ValueName##CameraParameter>(\
-				BlendableParameter, TargetStructProperty, CameraParameterPtr\
+				BlendableParameter, Target, TargetStructProperty, CameraParameterPtr\
 				);\
 	}\
 	else if (TargetStructProperty->Struct == F##ValueName##CameraVariableReference::StaticStruct())\
 	{\
 		auto* VariableReferencePtr = TargetStructProperty->ContainerPtrToValuePtr<F##ValueName##CameraVariableReference>(Target);\
 		Builder.SetCameraParameterOrVariableReferenceOverride<F##ValueName##CameraVariableReference>(\
-				BlendableParameter, TargetStructProperty, VariableReferencePtr\
+				BlendableParameter, Target, TargetStructProperty, VariableReferencePtr\
 				);\
 	}\
 	else
@@ -945,19 +962,19 @@ bool FCameraObjectInterfaceBuilder::SetupCameraParameterOrVariableReferenceOverr
 						"InvalidCameraNodeParameter",
 						"Invalid interface parameter '{0}', driving property '{1}' on '{2}', but it's not a camera parameter."),
 					FText::FromString(BlendableParameter->InterfaceParameterName), 
-					FText::FromName(BlendableParameter->TargetPropertyName),
+					FText::FromName(TargetPropertyName),
 					FText::FromName(Target->GetFName())));
 	}
 
 	return true;
 }
 
-bool FCameraObjectInterfaceBuilder::SetupCustomBlendableParameterOverride(const UCameraObjectInterfaceBlendableParameter* BlendableParameter)
+bool FCameraObjectInterfaceBuilder::SetupCustomBlendableParameterOverride(const UCameraObjectInterfaceBlendableParameter* BlendableParameter, UCameraNode* Target, FName TargetPropertyName)
 {
 	using namespace Internal;
 
-	ICustomCameraNodeParameterProvider* Target = Cast<ICustomCameraNodeParameterProvider>(BlendableParameter->Target);
-	if (!Target)
+	ICustomCameraNodeParameterProvider* TargetProvider = Cast<ICustomCameraNodeParameterProvider>(Target);
+	if (!TargetProvider)
 	{
 		// No match, try something else.
 		return false;
@@ -966,13 +983,13 @@ bool FCameraObjectInterfaceBuilder::SetupCustomBlendableParameterOverride(const
 	// Look for a parameter override matching the target name.
 	// TODO: we're querying the list of custom parameters every time, we may want to cache it for this phase.
 	FCustomCameraNodeParameterInfos CustomParameters;
-	Target->GetCustomCameraNodeParameters(CustomParameters);
+	TargetProvider->GetCustomCameraNodeParameters(CustomParameters);
 
 	FCustomCameraNodeParameterInfos::FBlendableParameterInfo* TargetCustomParameter = 
 		CustomParameters.BlendableParameters.FindByPredicate(
-			[BlendableParameter](FCustomCameraNodeParameterInfos::FBlendableParameterInfo& CustomParameter)
+			[TargetPropertyName](FCustomCameraNodeParameterInfos::FBlendableParameterInfo& CustomParameter)
 			{
-				return CustomParameter.ParameterName == BlendableParameter->TargetPropertyName;
+				return CustomParameter.ParameterName == TargetPropertyName;
 			});
 	if (!TargetCustomParameter)
 	{
@@ -985,25 +1002,25 @@ bool FCameraObjectInterfaceBuilder::SetupCustomBlendableParameterOverride(const
 	{
 #define UE_CAMERA_VARIABLE_FOR_TYPE(ValueType, ValueName)\
 		case ECameraVariableType::ValueName:\
-			Builder.SetCustomBlendableParameterOverride<U##ValueName##CameraVariable>(BlendableParameter, *TargetCustomParameter);\
+			Builder.SetCustomBlendableParameterOverride<U##ValueName##CameraVariable>(\
+					BlendableParameter, Target, TargetPropertyName, *TargetCustomParameter);\
 			break;
 	UE_CAMERA_VARIABLE_FOR_ALL_TYPES()
 #undef UE_CAMERA_VARIABLE_FOR_TYPE
 		case ECameraVariableType::BlendableStruct:
-			Builder.SetCustomBlendableStructParameterOverride(BlendableParameter, *TargetCustomParameter);
+			Builder.SetCustomBlendableStructParameterOverride(BlendableParameter, Target, TargetPropertyName, *TargetCustomParameter);
 			break;
 	}
 
 	return true;
 }
 
-bool FCameraObjectInterfaceBuilder::SetupDataContextPropertyOverride(const UCameraObjectInterfaceDataParameter* DataParameter)
+bool FCameraObjectInterfaceBuilder::SetupDataContextPropertyOverride(const UCameraObjectInterfaceDataParameter* DataParameter, UCameraNode* Target, FName TargetPropertyName)
 {
 	using namespace Internal;
 
-	UObject* Target = DataParameter->Target;
 	UClass* TargetClass = Target->GetClass();
-	FProperty* TargetProperty = TargetClass->FindPropertyByName(DataParameter->TargetPropertyName);
+	FProperty* TargetProperty = TargetClass->FindPropertyByName(TargetPropertyName);
 	if (!TargetProperty)
 	{
 		// No match, try something else.
@@ -1021,7 +1038,7 @@ bool FCameraObjectInterfaceBuilder::SetupDataContextPropertyOverride(const UCame
 						"Interface parameter '{0}' is driving data context property '{1}' on '{2}' "
 						"but no FCameraContextDataID property '{3}' was found to store the override ID."),
 					FText::FromString(DataParameter->InterfaceParameterName),
-					FText::FromName(DataParameter->TargetPropertyName),
+					FText::FromName(TargetPropertyName),
 					FText::FromName(Target->GetFName()),
 					FText::FromName(TargetDataIDPropertyName)));
 		return false;
@@ -1030,30 +1047,30 @@ bool FCameraObjectInterfaceBuilder::SetupDataContextPropertyOverride(const UCame
 	FCameraContextDataID* OverrideDataID = TargetDataIDProperty->ContainerPtrToValuePtr<FCameraContextDataID>(Target);
 
 	FInterfaceParameterBindingBuilder Builder(*this);
-	Builder.SetDataContextPropertyOverride(DataParameter, TargetProperty, OverrideDataID);
+	Builder.SetDataContextPropertyOverride(DataParameter, Target, TargetProperty, OverrideDataID);
 
 	return true;
 }
 
-bool FCameraObjectInterfaceBuilder::SetupCustomDataParameterOverride(const UCameraObjectInterfaceDataParameter* DataParameter)
+bool FCameraObjectInterfaceBuilder::SetupCustomDataParameterOverride(const UCameraObjectInterfaceDataParameter* DataParameter, UCameraNode* Target, FName TargetPropertyName)
 {
 	using namespace Internal;
 
-	ICustomCameraNodeParameterProvider* Target = Cast<ICustomCameraNodeParameterProvider>(DataParameter->Target);
-	if (!Target)
+	ICustomCameraNodeParameterProvider* TargetProvider = Cast<ICustomCameraNodeParameterProvider>(Target);
+	if (!TargetProvider)
 	{
 		// No match, try something else.
 		return false;
 	}
 
 	FCustomCameraNodeParameterInfos CustomParameters;
-	Target->GetCustomCameraNodeParameters(CustomParameters);
+	TargetProvider->GetCustomCameraNodeParameters(CustomParameters);
 
 	FCustomCameraNodeParameterInfos::FDataParameterInfo* TargetCustomParameter =
 		CustomParameters.DataParameters.FindByPredicate(
-				[DataParameter](FCustomCameraNodeParameterInfos::FDataParameterInfo& CustomParameter)
+				[TargetPropertyName](FCustomCameraNodeParameterInfos::FDataParameterInfo& CustomParameter)
 				{
-					return CustomParameter.ParameterName == DataParameter->TargetPropertyName;
+					return CustomParameter.ParameterName == TargetPropertyName;
 				});
 	if (!TargetCustomParameter)
 	{
@@ -1062,7 +1079,7 @@ bool FCameraObjectInterfaceBuilder::SetupCustomDataParameterOverride(const UCame
 	}
 
 	FInterfaceParameterBindingBuilder Builder(*this);
-	Builder.SetCustomDataParameterOverride(DataParameter, *TargetCustomParameter);
+	Builder.SetCustomDataParameterOverride(DataParameter, Target, TargetPropertyName, *TargetCustomParameter);
 
 	return true;
 }
diff --git a/Engine/Plugins/Cameras/GameplayCameras/Source/GameplayCameras/Private/Build/CameraObjectInterfaceParameterBuilder.cpp b/Engine/Plugins/Cameras/GameplayCameras/Source/GameplayCameras/Private/Build/CameraObjectInterfaceParameterBuilder.cpp
--- a/Engine/Plugins/Cameras/GameplayCameras/Source/GameplayCameras/Private/Build/CameraObjectInterfaceParameterBuilder.cpp
+++ b/Engine/Plugins/Cameras/GameplayCameras/Source/GameplayCameras/Private/Build/CameraObjectInterfaceParameterBuilder.cpp
@@ -42,12 +42,7 @@ void FCameraObjectInterfaceParameterBuilder::BuildParameterDefinitions()
 		if (BlendableParameter && BlendableParameter->PrivateVariableID)
 		{
 			FCameraObjectInterfaceParameterDefinition Definition;
-			Definition.ParameterName = FName(BlendableParameter->InterfaceParameterName);
-			Definition.ParameterGuid = BlendableParameter->GetGuid();
-			Definition.ParameterType = ECameraObjectInterfaceParameterType::Blendable;
-			Definition.VariableID = BlendableParameter->PrivateVariableID;
-			Definition.VariableType = BlendableParameter->ParameterType;
-			Definition.BlendableStructType = BlendableParameter->BlendableStructType;
+			BlendableParameter->GetParameterDefinition(Definition);
 			ParameterDefinitions.Add(Definition);
 		}
 	}
@@ -57,13 +52,7 @@ void FCameraObjectInterfaceParameterBuilder::BuildParameterDefinitions()
 		if (DataParameter && DataParameter->PrivateDataID)
 		{
 			FCameraObjectInterfaceParameterDefinition Definition;
-			Definition.ParameterName = FName(DataParameter->InterfaceParameterName);
-			Definition.ParameterGuid = DataParameter->GetGuid();
-			Definition.ParameterType = ECameraObjectInterfaceParameterType::Data;
-			Definition.DataID = DataParameter->PrivateDataID;
-			Definition.DataType = DataParameter->DataType;
-			Definition.DataContainerType = DataParameter->DataContainerType;
-			Definition.DataTypeObject = DataParameter->DataTypeObject;
+			DataParameter->GetParameterDefinition(Definition);
 			ParameterDefinitions.Add(Definition);
 		}
 	}
@@ -77,23 +66,16 @@ void FCameraObjectInterfaceParameterBuilder::BuildParameterDefinitions()
 
 void FCameraObjectInterfaceParameterBuilder::BuildDefaultParameters()
 {
-	FInstancedPropertyBag DefaultParameters;
-	BuildDefaultParameters(CameraObject, DefaultParameters);
-	if (!DefaultParameters.Identical(&CameraObject->DefaultParameters, 0))
+	TArray<FPropertyBagPropertyDesc> DefaultParameterProperties;
+	AppendDefaultParameterProperties(CameraObject, DefaultParameterProperties);
+	const UPropertyBag* DefaultParametersStruct = UPropertyBag::GetOrCreateFromDescs(DefaultParameterProperties);
+	if (CameraObject->DefaultParameters.GetPropertyBagStruct() != DefaultParametersStruct)
 	{
 		CameraObject->Modify();
-		CameraObject->DefaultParameters = DefaultParameters;
+		CameraObject->DefaultParameters.MigrateToNewBagStruct(DefaultParametersStruct);
 	}
 }
 
-void FCameraObjectInterfaceParameterBuilder::BuildDefaultParameters(const UBaseCameraObject* CameraObject, FInstancedPropertyBag& OutPropertyBag)
-{
-	TArray<FPropertyBagPropertyDesc> DefaultParameterProperties;
-	AppendDefaultParameterProperties(CameraObject, DefaultParameterProperties);
-	OutPropertyBag.AddProperties(DefaultParameterProperties);
-	SetDefaultParameterValues(CameraObject, OutPropertyBag);
-}
-
 void FCameraObjectInterfaceParameterBuilder::AppendDefaultParameterProperties(const UBaseCameraObject* CameraObject, TArray<FPropertyBagPropertyDesc>& OutProperties)
 {
 	AppendDefaultParameterProperties(CameraObject->GetParameterDefinitions(), OutProperties);
@@ -183,6 +165,10 @@ void FCameraObjectInterfaceParameterBuilder::AppendDefaultParameterProperties(TC
 			NewProperty.ID = Definition.ParameterGuid;
 #if UE_VERSION_NEWER_THAN_OR_EQUAL(5,6,0)
 			NewProperty.PropertyFlags |= PropertyFlags;
+			if (!Definition.bIsVisible)
+			{
+				NewProperty.PropertyFlags &= ~CPF_Edit;
+			}
 #endif
 
 			OutProperties.Add(NewProperty);
@@ -190,107 +176,159 @@ void FCameraObjectInterfaceParameterBuilder::AppendDefaultParameterProperties(TC
 	}
 }
 
-void FCameraObjectInterfaceParameterBuilder::SetDefaultParameterValues(const UBaseCameraObject* CameraObject, FInstancedPropertyBag& PropertyBag)
+bool FCameraObjectInterfaceParameterBuilder::SetDefaultParameterValue(UBaseCameraObject* CameraObject, const FCameraObjectInterfaceParameterDefinition& ParameterDefinition, UCameraNode* TargetNode, FName TargetPropertyName, bool bAddParameterIfMissing)
 {
-	uint8* PropertyBagValue = PropertyBag.GetMutableValue().GetMemory();
-	const UPropertyBag* PropertyBagStruct = PropertyBag.GetPropertyBagStruct();
-	if (!ensure(PropertyBagValue && PropertyBagStruct))
+	if (!ensure(CameraObject && TargetNode && TargetPropertyName != NAME_None))
 	{
-		return;
+		return false;
 	}
 
-	for (const UCameraObjectInterfaceBlendableParameter* BlendableParameter : CameraObject->Interface.BlendableParameters)
-	{
-		if (!ensure(BlendableParameter))
-		{
-			continue;
-		}
+	// Start by finding the source value, i.e. the value of the camera node property being exposed as a parameter.
+	const void* RawSourceValuePtr = nullptr;
 
-		UCameraNode* CameraNode = BlendableParameter->Target;
-		if (!CameraNode)
+	// First check if the value is found on a custom parameter.
+	if (ICustomCameraNodeParameterProvider* CustomParameterProvider = Cast<ICustomCameraNodeParameterProvider>(TargetNode))
+	{
+		FCustomCameraNodeParameterInfos CustomParameters;
+		CustomParameterProvider->GetCustomCameraNodeParameters(CustomParameters);
+
+		FCustomCameraNodeParameterInfos::FBlendableParameterInfo* CustomBlendableParameter = 
+			CustomParameters.BlendableParameters.FindByPredicate(
+					[TargetPropertyName](FCustomCameraNodeParameterInfos::FBlendableParameterInfo& CustomParameter)
+					{
+						return CustomParameter.ParameterName == TargetPropertyName;
+					});
+		if (CustomBlendableParameter)
 		{
-			continue;
+			RawSourceValuePtr = CustomBlendableParameter->DefaultValue;
+			goto DoneSearchingRawSourceValuePtr;
 		}
 
-		const void* RawSourceValuePtr = nullptr;
-
-		// First check if the value is found on a custom parameter.
-		if (ICustomCameraNodeParameterProvider* CustomParameterProvider = Cast<ICustomCameraNodeParameterProvider>(CameraNode))
+		FCustomCameraNodeParameterInfos::FDataParameterInfo* CustomDataParameter =
+			CustomParameters.DataParameters.FindByPredicate(
+					[TargetPropertyName](FCustomCameraNodeParameterInfos::FDataParameterInfo& CustomParameter)
+					{
+						return CustomParameter.ParameterName == TargetPropertyName;
+					});
+		if (CustomDataParameter)
 		{
-			FCustomCameraNodeParameterInfos CustomParameters;
-			CustomParameterProvider->GetCustomCameraNodeParameters(CustomParameters);
-
-			FCustomCameraNodeParameterInfos::FBlendableParameterInfo* TargetCustomParameter = 
-				CustomParameters.BlendableParameters.FindByPredicate(
-						[BlendableParameter](FCustomCameraNodeParameterInfos::FBlendableParameterInfo& CustomParameter)
-						{
-							return CustomParameter.ParameterName == BlendableParameter->TargetPropertyName;
-						});
-			if (TargetCustomParameter)
-			{
-				RawSourceValuePtr = TargetCustomParameter->DefaultValue;
-			}
+			RawSourceValuePtr = CustomDataParameter->DefaultValue;
+			goto DoneSearchingRawSourceValuePtr;
 		}
+	}
 
-		// If not found, check on a reflected UObject property.
-		if (!RawSourceValuePtr)
+	// If not found, check on a reflected UObject property.
+	// Blendable parameters would be hooked to a camera parameter struct, or a blendable struct.
+	// Data parameters could be hooked up to anything.
+	if (!RawSourceValuePtr)
+	{
+		const UClass* TargetClass = TargetNode->GetClass();
+		FProperty* TargetProperty = TargetClass->FindPropertyByName(TargetPropertyName);
+		FStructProperty* StructProperty = CastField<FStructProperty>(TargetProperty);
+		if (StructProperty && ParameterDefinition.ParameterType == ECameraObjectInterfaceParameterType::Blendable)
 		{
-			const UClass* TargetClass = CameraNode->GetClass();
-			FStructProperty* StructProperty = CastField<FStructProperty>(
-					TargetClass->FindPropertyByName(BlendableParameter->TargetPropertyName));
-			if (StructProperty)
+			switch (ParameterDefinition.VariableType)
 			{
-				switch (BlendableParameter->ParameterType)
-				{
 #define UE_CAMERA_VARIABLE_FOR_TYPE(ValueType, ValueName)\
-					case ECameraVariableType::ValueName:\
+				case ECameraVariableType::ValueName:\
+					{\
+						using CameraParameterType = F##ValueName##CameraParameter;\
+						using CameraVariableReferenceType = F##ValueName##CameraVariableReference;\
+						if (StructProperty->Struct == CameraParameterType::StaticStruct())\
+						{\
+							CameraParameterType* CameraParameterPtr = StructProperty->ContainerPtrToValuePtr<CameraParameterType>(TargetNode);\
+							RawSourceValuePtr = static_cast<void*>(&CameraParameterPtr->Value);\
+						}\
+						else if (StructProperty->Struct == CameraVariableReferenceType::StaticStruct())\
 						{\
-							using CameraParameterType = F##ValueName##CameraParameter;\
-							using CameraVariableReferenceType = F##ValueName##CameraVariableReference;\
-							if (StructProperty->Struct == CameraParameterType::StaticStruct())\
-							{\
-								CameraParameterType* CameraParameterPtr = StructProperty->ContainerPtrToValuePtr<CameraParameterType>(CameraNode);\
-								RawSourceValuePtr = static_cast<void*>(&CameraParameterPtr->Value);\
-							}\
-							else if (StructProperty->Struct == CameraVariableReferenceType::StaticStruct())\
-							{\
-								CameraVariableReferenceType* VariableReferencePtr = StructProperty->ContainerPtrToValuePtr<CameraVariableReferenceType>(CameraNode);\
-								RawSourceValuePtr = VariableReferencePtr->Variable ? VariableReferencePtr->Variable->GetDefaultValuePtr() : nullptr;\
-							}\
+							CameraVariableReferenceType* VariableReferencePtr = StructProperty->ContainerPtrToValuePtr<CameraVariableReferenceType>(TargetNode);\
+							RawSourceValuePtr = VariableReferencePtr->Variable ? VariableReferencePtr->Variable->GetDefaultValuePtr() : nullptr;\
 						}\
-						break;
-					UE_CAMERA_VARIABLE_FOR_ALL_TYPES()
+					}\
+					break;
+				UE_CAMERA_VARIABLE_FOR_ALL_TYPES()
 #undef UE_CAMERA_VARIABLE_FOR_TYPE
-					case ECameraVariableType::BlendableStruct:
-						{
-							RawSourceValuePtr = StructProperty->ContainerPtrToValuePtr<void>(CameraNode);
-						}
-						break;
-				}
+				case ECameraVariableType::BlendableStruct:
+					{
+						RawSourceValuePtr = StructProperty->ContainerPtrToValuePtr<void>(TargetNode);
+					}
+					break;
 			}
 		}
+		else if (ParameterDefinition.ParameterType == ECameraObjectInterfaceParameterType::Data)
+		{
+			RawSourceValuePtr = TargetProperty->ContainerPtrToValuePtr<void>(TargetNode);
+		}
+	}
 
-		if (!RawSourceValuePtr)
+DoneSearchingRawSourceValuePtr:
+
+	if (!RawSourceValuePtr)
+	{
+		return false;
+	}
+
+	// Second, find the corresponding property on the default parameters' property bag. We will write the source 
+	// value to this property. If we don't find it, either bail out, or add it to the property bag structure and 
+	// continue.
+	FInstancedPropertyBag& DefaultParameters = CameraObject->DefaultParameters;
+	const UPropertyBag* PropertyBagStruct = DefaultParameters.GetPropertyBagStruct();
+	if (!PropertyBagStruct)
+	{
+		if (!bAddParameterIfMissing)
 		{
-			continue;
+			return false;
 		}
+		// else, we will (re)create the property bag struct with the new property below.
+	}
 
-		// Find the corresponding property on the default parameters' property bag.
-		const FPropertyBagPropertyDesc* PropertyDesc = PropertyBagStruct->FindPropertyDescByID(BlendableParameter->GetGuid());
-		if (!ensure(PropertyDesc && PropertyDesc->CachedProperty))
+	const FPropertyBagPropertyDesc* PropertyDesc = PropertyBagStruct ? 
+		PropertyBagStruct->FindPropertyDescByID(ParameterDefinition.ParameterGuid) :
+		nullptr;
+	if (!PropertyDesc)
+	{
+		if (!bAddParameterIfMissing)
 		{
-			continue;
+			return false;
+		}
+
+		TArray<FPropertyBagPropertyDesc> NewProperties;
+		TConstArrayView<FCameraObjectInterfaceParameterDefinition> NewParameterDefinitions = MakeConstArrayView(&ParameterDefinition, 1);
+		AppendDefaultParameterProperties(NewParameterDefinitions, NewProperties);
+		const EPropertyBagAlterationResult Result = DefaultParameters.AddProperties(NewProperties, false);
+		if (!ensure(Result == EPropertyBagAlterationResult::Success))
+		{
+			return false;
 		}
 
-		// This property should be a structure: either a camera parameter for all the standard blendable types,
-		// or a blendable structure.
+		// Re-query the struct, since it changed when we added the property.
+		PropertyBagStruct = DefaultParameters.GetPropertyBagStruct();
+		PropertyDesc = PropertyBagStruct->FindPropertyDescByID(ParameterDefinition.ParameterGuid);
+	}
+
+	if (!ensure(PropertyDesc && PropertyDesc->CachedProperty))
+	{
+		return false;
+	}
+
+	uint8* PropertyBagValue = DefaultParameters.GetMutableValue().GetMemory();
+	if (!ensure(PropertyBagValue && PropertyBagStruct))
+	{
+		return false;
+	}
+
+	// The given property should be a structure for blendable parameters (either a camera parameter or a blendable 
+	// structure). For data parameters, it can be anything we support. Now do the value copying!
+	bool bSuccess = true;
+	if (ParameterDefinition.ParameterType == ECameraObjectInterfaceParameterType::Blendable)
+	{
 		const FStructProperty* DefaultParameterProperty = CastField<FStructProperty>(PropertyDesc->CachedProperty);
 		if (!ensure(DefaultParameterProperty))
 		{
-			continue;
+			return false;
 		}
 
-		switch (BlendableParameter->ParameterType)
+		switch (ParameterDefinition.VariableType)
 		{
 #define UE_CAMERA_VARIABLE_FOR_TYPE(ValueType, ValueName)\
 			case ECameraVariableType::ValueName:\
@@ -307,109 +345,64 @@ void FCameraObjectInterfaceParameterBuilder::SetDefaultParameterValues(const UBa
 			UE_CAMERA_VARIABLE_FOR_ALL_TYPES()
 #undef UE_CAMERA_VARIABLE_FOR_TYPE
 			case ECameraVariableType::BlendableStruct:
-				if (ensure(DefaultParameterProperty->Struct == BlendableParameter->BlendableStructType))
+				if (ensure(DefaultParameterProperty->Struct == ParameterDefinition.BlendableStructType))
 				{
 					void* RawDestinationValuePtr = DefaultParameterProperty->ContainerPtrToValuePtr<void>(PropertyBagValue);
-					BlendableParameter->BlendableStructType->CopyScriptStruct(RawDestinationValuePtr, RawSourceValuePtr);
+					ParameterDefinition.BlendableStructType->CopyScriptStruct(RawDestinationValuePtr, RawSourceValuePtr);
 				}
 				break;
+			default:
+				ensureMsgf(false, TEXT("Unsupported or unknown variable type!"));
+				bSuccess = false;
+				break;
 		}
 	}
-
-	for (const UCameraObjectInterfaceDataParameter* DataParameter : CameraObject->Interface.DataParameters)
+	else if (ParameterDefinition.ParameterType == ECameraObjectInterfaceParameterType::Data)
 	{
-		if (!ensure(DataParameter))
-		{
-			continue;
-		}
-
-		UCameraNode* CameraNode = DataParameter->Target;
-		if (!CameraNode)
-		{
-			continue;
-		}
-
-		const uint8* RawSourceValuePtr = nullptr;
-
-		if (ICustomCameraNodeParameterProvider* CustomParameterProvider = Cast<ICustomCameraNodeParameterProvider>(CameraNode))
+		if (ParameterDefinition.DataContainerType == ECameraContextDataContainerType::None)
 		{
-			FCustomCameraNodeParameterInfos CustomParameters;
-			CustomParameterProvider->GetCustomCameraNodeParameters(CustomParameters);
-
-			FCustomCameraNodeParameterInfos::FDataParameterInfo* TargetCustomParameter = 
-				CustomParameters.DataParameters.FindByPredicate(
-						[DataParameter](FCustomCameraNodeParameterInfos::FDataParameterInfo& CustomParameter)
-						{
-							return CustomParameter.ParameterName == DataParameter->TargetPropertyName;
-						});
-			if (TargetCustomParameter)
+			void* RawDestinationValuePtr = PropertyDesc->CachedProperty->ContainerPtrToValuePtr<void>(PropertyBagValue);
+			if (!ensure(RawDestinationValuePtr))
 			{
-				RawSourceValuePtr = TargetCustomParameter->DefaultValue;
+				return false;
 			}
-		}
-
-		if (!RawSourceValuePtr)
-		{
-			const UClass* TargetClass = DataParameter->Target->GetClass();
-			FProperty* Property = TargetClass->FindPropertyByName(DataParameter->TargetPropertyName);
-			if (Property)
+			
+			switch (ParameterDefinition.DataType)
 			{
-				RawSourceValuePtr = (uint8*)Property->ContainerPtrToValuePtr<void>(DataParameter->Target);
-			}
-		}
-
-		if (!ensure(RawSourceValuePtr))
-		{
-			continue;
-		}
-
-		const FPropertyBagPropertyDesc* PropertyDesc = PropertyBagStruct->FindPropertyDescByID(DataParameter->GetGuid());
-		if (!ensure(PropertyDesc && PropertyDesc->CachedProperty))
-		{
-			continue;
-		}
-
-		if (DataParameter->DataContainerType == ECameraContextDataContainerType::None)
-		{
-			void* RawDestinationValuePtr = PropertyDesc->CachedProperty->ContainerPtrToValuePtr<void>(PropertyBagValue);
-			if (ensure(RawDestinationValuePtr))
-			{
-				SetDefaultParameterValue(DataParameter, RawDestinationValuePtr, RawSourceValuePtr);
+				case ECameraContextDataType::Name:
+					*((FName*)RawDestinationValuePtr) = *((FName*)RawSourceValuePtr);
+					break;
+				case ECameraContextDataType::String:
+					*((FString*)RawDestinationValuePtr) = *((FString*)RawSourceValuePtr);
+					break;
+				case ECameraContextDataType::Enum:
+					*((uint8*)RawDestinationValuePtr) = *((uint8*)RawSourceValuePtr);
+					break;
+				case ECameraContextDataType::Struct:
+					{
+						const UScriptStruct* StructType = CastChecked<const UScriptStruct>(ParameterDefinition.DataTypeObject);
+						StructType->CopyScriptStruct(RawDestinationValuePtr, RawSourceValuePtr);
+					}
+					break;
+				case ECameraContextDataType::Object:
+					*((FObjectPtr*)RawDestinationValuePtr) = *((FObjectPtr*)RawSourceValuePtr);
+					break;
+				case ECameraContextDataType::Class:
+					*((FObjectPtr*)RawDestinationValuePtr) = *((FObjectPtr*)RawSourceValuePtr);
+					break;
+				default:
+					ensureMsgf(false, TEXT("Unsupported or unknown data type!"));
+					bSuccess = false;
+					break;
 			}
 		}
-		else if (DataParameter->DataContainerType == ECameraContextDataContainerType::Array)
+		else if (ParameterDefinition.DataContainerType == ECameraContextDataContainerType::Array)
 		{
 			// Array properties are empty by default.
 		}
 	}
-}
 
-void FCameraObjectInterfaceParameterBuilder::SetDefaultParameterValue(const UCameraObjectInterfaceDataParameter* DataParameter, void* DestValuePtr, const void* SrcValuePtr)
-{
-	switch (DataParameter->DataType)
-	{
-		case ECameraContextDataType::Name:
-			*((FName*)DestValuePtr) = *((FName*)SrcValuePtr);
-			break;
-		case ECameraContextDataType::String:
-			*((FString*)DestValuePtr) = *((FString*)SrcValuePtr);
-			break;
-		case ECameraContextDataType::Enum:
-			*((uint8*)DestValuePtr) = *((uint8*)SrcValuePtr);
-			break;
-		case ECameraContextDataType::Struct:
-			{
-				const UScriptStruct* StructType = CastChecked<const UScriptStruct>(DataParameter->DataTypeObject);
-				StructType->CopyScriptStruct(DestValuePtr, SrcValuePtr);
-			}
-			break;
-		case ECameraContextDataType::Object:
-			*((FObjectPtr*)DestValuePtr) = *((FObjectPtr*)SrcValuePtr);
-			break;
-		case ECameraContextDataType::Class:
-			*((FObjectPtr*)DestValuePtr) = *((FObjectPtr*)SrcValuePtr);
-			break;
-	}
+	return bSuccess;
 }
 
 #if UE_VERSION_OLDER_THAN(5,8,0)
diff --git a/Engine/Plugins/Cameras/GameplayCameras/Source/GameplayCameras/Private/Core/BaseCameraObject.cpp b/Engine/Plugins/Cameras/GameplayCameras/Source/GameplayCameras/Private/Core/BaseCameraObject.cpp
--- a/Engine/Plugins/Cameras/GameplayCameras/Source/GameplayCameras/Private/Core/BaseCameraObject.cpp
+++ b/Engine/Plugins/Cameras/GameplayCameras/Source/GameplayCameras/Private/Core/BaseCameraObject.cpp
@@ -3,6 +3,8 @@
 #include "Core/BaseCameraObject.h"
 
 #include "Build/CameraObjectInterfaceParameterBuilder.h"
+#include "Core/CameraNode.h"
+#include "Core/ObjectTreeGraphRootObject.h"
 #include "Misc/EngineVersionComparison.h"
 
 #include UE_INLINE_GENERATED_CPP_BY_NAME(BaseCameraObject)
@@ -20,6 +22,48 @@ void FCameraObjectAllocationInfo::Append(const FCameraObjectAllocationInfo& Othe
 	ContextDataTableInfo.DataDefinitions.Append(OtherContextDataTableInfo.DataDefinitions);
 }
 
+void FCameraObjectConnections::Add(UObject* InSource, FName InSourcePropertyName, UObject* InTarget, FName InTargetPropertyName)
+{
+	FCameraObjectConnection NewConnection;
+	NewConnection.Source = InSource;
+	NewConnection.SourcePropertyName = InSourcePropertyName;
+	NewConnection.Target = InTarget;
+	NewConnection.TargetPropertyName = InTargetPropertyName;
+	Connections.Add(NewConnection);
+}
+
+FCameraObjectConnection* FCameraObjectConnections::FindBySource(UObject* InSource)
+{
+	return Connections.FindByPredicate([InSource](FCameraObjectConnection& Item)
+			{
+				return Item.Source == InSource;
+			});
+}
+
+FCameraObjectConnection* FCameraObjectConnections::FindBySource(UObject* InSource, FName InSourcePropertyName)
+{
+	return Connections.FindByPredicate([InSource, InSourcePropertyName](FCameraObjectConnection& Item)
+			{
+				return Item.Source == InSource && Item.SourcePropertyName == InSourcePropertyName;
+			});
+}
+
+FCameraObjectConnection* FCameraObjectConnections::FindByTarget(UObject* InTarget)
+{
+	return Connections.FindByPredicate([InTarget](FCameraObjectConnection& Item)
+			{
+				return Item.Target == InTarget;
+			});
+}
+
+FCameraObjectConnection* FCameraObjectConnections::FindByTarget(UObject* InTarget, FName InTargetPropertyName)
+{
+	return Connections.FindByPredicate([InTarget, InTargetPropertyName](FCameraObjectConnection& Item)
+			{
+				return Item.Target == InTarget && Item.TargetPropertyName == InTargetPropertyName;
+			});
+}
+
 void UBaseCameraObject::PostLoad()
 {
 	Super::PostLoad();
@@ -56,3 +100,61 @@ void UBaseCameraObject::PostDuplicate(EDuplicateMode::Type DuplicateMode)
 	}
 }
 
+#if WITH_EDITORONLY_DATA
+
+void UBaseCameraObject::UpgradeInterfaceConnections(IObjectTreeGraphRootObject* RootObject, FName DefaultGraphName)
+{
+	if (!ensure(RootObject))
+	{
+		return;
+	}
+
+	TArray<UCameraObjectInterfaceParameterBase*> AllParameters;
+	AllParameters.Append(Interface.BlendableParameters);
+	AllParameters.Append(Interface.DataParameters);
+
+	// Move the connection information to the new connection list.
+	// Create getter nodes for parameters that were added to the camera node graph.
+	bool bDidModify = false;
+	for (UCameraObjectInterfaceParameterBase* Parameter : AllParameters)
+	{
+		UCameraObjectInterfaceParameterGetter* GetterNode = nullptr;
+		const bool bIsConnected = (Parameter->Target_DEPRECATED && Parameter->TargetPropertyName_DEPRECATED != NAME_None);
+
+		if (Parameter->bHasGraphNode_DEPRECATED || bIsConnected)
+		{
+			GetterNode = NewObject<UCameraObjectInterfaceParameterGetter>(this);
+			GetterNode->ParameterGuid = Parameter->GetGuid();
+			GetterNode->GraphNodePos = Parameter->GraphNodePos_DEPRECATED;
+			RootObject->AddConnectableObject(DefaultGraphName, GetterNode);
+			bDidModify = true;
+
+			Parameter->bHasGraphNode_DEPRECATED = false;
+			Parameter->GraphNodePos_DEPRECATED = FIntVector2::ZeroValue;
+		}
+
+		if (bIsConnected)
+		{
+			ensure(GetterNode);
+
+			FCameraObjectConnection Connection;
+			Connection.Source = GetterNode;
+			Connection.Target = Parameter->Target_DEPRECATED;
+			Connection.TargetPropertyName = Parameter->TargetPropertyName_DEPRECATED;
+
+			Connections.Connections.Add(Connection);
+			bDidModify = true;
+
+			Parameter->Target_DEPRECATED = nullptr;
+			Parameter->TargetPropertyName_DEPRECATED = NAME_None;
+		}
+	}
+
+	if (bDidModify)
+	{
+		Modify();
+	}
+}
+
+#endif  // WITH_EDITORONLY_DATA
+
diff --git a/Engine/Plugins/Cameras/GameplayCameras/Source/GameplayCameras/Private/Core/BaseCameraObjectReference.cpp b/Engine/Plugins/Cameras/GameplayCameras/Source/GameplayCameras/Private/Core/BaseCameraObjectReference.cpp
--- a/Engine/Plugins/Cameras/GameplayCameras/Source/GameplayCameras/Private/Core/BaseCameraObjectReference.cpp
+++ b/Engine/Plugins/Cameras/GameplayCameras/Source/GameplayCameras/Private/Core/BaseCameraObjectReference.cpp
@@ -151,6 +151,11 @@ void FBaseCameraObjectReference::GetCustomCameraNodeParameters(FCustomCameraNode
 
 	for (const FCameraObjectInterfaceParameterDefinition& Definition : CameraObject->GetParameterDefinitions())
 	{
+		if (!Definition.bIsVisible)
+		{
+			continue;
+		}
+
 		const FPropertyBagPropertyDesc* PropertyDesc = ParametersStruct->FindPropertyDescByID(Definition.ParameterGuid);
 		if (!ensure(PropertyDesc && PropertyDesc->CachedProperty))
 		{
diff --git a/Engine/Plugins/Cameras/GameplayCameras/Source/GameplayCameras/Private/Core/CameraObjectInterface.cpp b/Engine/Plugins/Cameras/GameplayCameras/Source/GameplayCameras/Private/Core/CameraObjectInterface.cpp
--- a/Engine/Plugins/Cameras/GameplayCameras/Source/GameplayCameras/Private/Core/CameraObjectInterface.cpp
+++ b/Engine/Plugins/Cameras/GameplayCameras/Source/GameplayCameras/Private/Core/CameraObjectInterface.cpp
@@ -2,29 +2,12 @@
 
 #include "Core/CameraObjectInterface.h"
 
+#include "Core/BaseCameraObject.h"
 #include "Core/CameraContextDataTableAllocationInfo.h"
 #include "Core/CameraVariableTableAllocationInfo.h"
 
 #include UE_INLINE_GENERATED_CPP_BY_NAME(CameraObjectInterface)
 
-#if WITH_EDITOR
-
-void UCameraObjectInterfaceParameterBase::GetGraphNodePosition(FName InGraphName, int32& NodePosX, int32& NodePosY) const
-{
-	NodePosX = GraphNodePos.X;
-	NodePosY = GraphNodePos.Y;
-}
-
-void UCameraObjectInterfaceParameterBase::OnGraphNodeMoved(FName InGraphName, int32 NodePosX, int32 NodePosY, bool bMarkDirty)
-{
-	Modify(bMarkDirty);
-
-	GraphNodePos.X = NodePosX;
-	GraphNodePos.Y = NodePosY;
-}
-
-#endif
-
 void UCameraObjectInterfaceParameterBase::PostLoad()
 {
 	if (!Guid.IsValid())
@@ -70,6 +53,17 @@ FCameraVariableDefinition UCameraObjectInterfaceBlendableParameter::GetVariableD
 	return Definition;
 }
 
+void UCameraObjectInterfaceBlendableParameter::GetParameterDefinition(FCameraObjectInterfaceParameterDefinition& OutParameterDefinition) const
+{
+	OutParameterDefinition.ParameterName = FName(InterfaceParameterName);
+	OutParameterDefinition.ParameterGuid = Guid;
+	OutParameterDefinition.ParameterType = ECameraObjectInterfaceParameterType::Blendable;
+	OutParameterDefinition.bIsVisible = bIsVisible;
+	OutParameterDefinition.VariableID = PrivateVariableID;
+	OutParameterDefinition.VariableType = ParameterType;
+	OutParameterDefinition.BlendableStructType = BlendableStructType;
+}
+
 FCameraContextDataDefinition UCameraObjectInterfaceDataParameter::GetDataDefinition() const
 {
 	FCameraContextDataDefinition Definition;
@@ -83,6 +77,18 @@ FCameraContextDataDefinition UCameraObjectInterfaceDataParameter::GetDataDefinit
 	return Definition;
 }
 
+void UCameraObjectInterfaceDataParameter::GetParameterDefinition(FCameraObjectInterfaceParameterDefinition& OutParameterDefinition) const
+{
+	OutParameterDefinition.ParameterName = FName(InterfaceParameterName);
+	OutParameterDefinition.ParameterGuid = Guid;
+	OutParameterDefinition.ParameterType = ECameraObjectInterfaceParameterType::Data;
+	OutParameterDefinition.bIsVisible = bIsVisible;
+	OutParameterDefinition.DataID = PrivateDataID;
+	OutParameterDefinition.DataType = DataType;
+	OutParameterDefinition.DataContainerType = DataContainerType;
+	OutParameterDefinition.DataTypeObject = DataTypeObject;
+}
+
 #if WITH_EDITORONLY_DATA
 
 FString UCameraObjectInterfaceBlendableParameter::GetVariableName() const
@@ -144,3 +150,49 @@ bool FCameraObjectInterface::HasBlendableParameter(const FString& ParameterName)
 	return FindBlendableParameterByName(ParameterName) != nullptr;
 }
 
+UCameraObjectInterfaceParameterBase* UCameraObjectInterfaceParameterGetter::GetInterfaceParameter() const
+{
+	if (ParameterGuid.IsValid())
+	{
+		if (UBaseCameraObject* CameraObject = GetTypedOuter<UBaseCameraObject>())
+		{
+			if (UCameraObjectInterfaceBlendableParameter* BlendableParameter = CameraObject->Interface.FindBlendableParameterByGuid(ParameterGuid))
+			{
+				return BlendableParameter;
+			}
+			if (UCameraObjectInterfaceDataParameter* DataParameter = CameraObject->Interface.FindDataParameterByGuid(ParameterGuid))
+			{
+				return DataParameter;
+			}
+		}
+	}
+	return nullptr;
+}
+
+FString UCameraObjectInterfaceParameterGetter::GetInterfaceParameterName() const
+{
+	if (UCameraObjectInterfaceParameterBase* Parameter = GetInterfaceParameter())
+	{
+		return Parameter->InterfaceParameterName;
+	}
+	return FString();
+}
+
+#if WITH_EDITOR
+
+void UCameraObjectInterfaceParameterGetter::GetGraphNodePosition(FName InGraphName, int32& NodePosX, int32& NodePosY) const
+{
+	NodePosX = GraphNodePos.X;
+	NodePosY = GraphNodePos.Y;
+}
+
+void UCameraObjectInterfaceParameterGetter::OnGraphNodeMoved(FName InGraphName, int32 NodePosX, int32 NodePosY, bool bMarkDirty)
+{
+	Modify(bMarkDirty);
+
+	GraphNodePos.X = NodePosX;
+	GraphNodePos.Y = NodePosY;
+}
+
+#endif
+
diff --git a/Engine/Plugins/Cameras/GameplayCameras/Source/GameplayCameras/Private/Core/CameraRigAsset.cpp b/Engine/Plugins/Cameras/GameplayCameras/Source/GameplayCameras/Private/Core/CameraRigAsset.cpp
--- a/Engine/Plugins/Cameras/GameplayCameras/Source/GameplayCameras/Private/Core/CameraRigAsset.cpp
+++ b/Engine/Plugins/Cameras/GameplayCameras/Source/GameplayCameras/Private/Core/CameraRigAsset.cpp
@@ -43,7 +43,10 @@ void UCameraRigAsset::PostLoad()
 		}
 		if (UCameraObjectInterfaceParameterBase* InterfaceParameter = Cast<UCameraObjectInterfaceParameterBase>(Item))
 		{
-			InterfaceParameter->bHasGraphNode = true;
+			// This will get further upgraded below in UpgradeInterfaceConnections.
+			PRAGMA_DISABLE_DEPRECATION_WARNINGS
+			InterfaceParameter->bHasGraphNode_DEPRECATED = true;
+			PRAGMA_ENABLE_DEPRECATION_WARNINGS
 			It.RemoveCurrent();
 			continue;
 		}
@@ -64,6 +67,11 @@ void UCameraRigAsset::PostLoad()
 			}
 		}
 	}
+
+#if WITH_EDITORONLY_DATA
+	// Upgrade the old parameter connections.
+	UpgradeInterfaceConnections(this, NodeTreeGraphName);
+#endif
 }
 
 void UCameraRigAsset::GetOwnedGameplayTags(FGameplayTagContainer& TagContainer) const
diff --git a/Engine/Plugins/Cameras/GameplayCameras/Source/GameplayCameras/Private/Core/CameraShakeAsset.cpp b/Engine/Plugins/Cameras/GameplayCameras/Source/GameplayCameras/Private/Core/CameraShakeAsset.cpp
--- a/Engine/Plugins/Cameras/GameplayCameras/Source/GameplayCameras/Private/Core/CameraShakeAsset.cpp
+++ b/Engine/Plugins/Cameras/GameplayCameras/Source/GameplayCameras/Private/Core/CameraShakeAsset.cpp
@@ -38,6 +38,16 @@ void UCameraShakeAsset::DirtyBuildStatus()
 	BuildStatus = ECameraBuildStatus::Dirty;
 }
 
+void UCameraShakeAsset::PostLoad()
+{
+	Super::PostLoad();
+
+#if WITH_EDITORONLY_DATA
+	// Upgrade the old parameter connections.
+	UpgradeInterfaceConnections(this, NAME_None);
+#endif
+}
+
 void UCameraShakeAsset::PreSave(FObjectPreSaveContext ObjectSaveContext)
 {
 #if WITH_EDITOR
diff --git a/Engine/Plugins/Cameras/GameplayCameras/Source/GameplayCameras/Private/Tests/GameplayCamerasTestBuilder.h b/Engine/Plugins/Cameras/GameplayCameras/Source/GameplayCameras/Private/Tests/GameplayCamerasTestBuilder.h
--- a/Engine/Plugins/Cameras/GameplayCameras/Source/GameplayCameras/Private/Tests/GameplayCamerasTestBuilder.h
+++ b/Engine/Plugins/Cameras/GameplayCameras/Source/GameplayCameras/Private/Tests/GameplayCamerasTestBuilder.h
@@ -400,7 +400,7 @@ private:
  *					.SetParameter(&UOffsetCameraNode::TranslationOffset, FVector3d{ 1, 0, 0 })
  *					.Done()
  *				.AddChild<ULensParametersCameraNode>(&UArrayCameraNode::Children)
- *					.SetParameter(&ULensParametersCameraNode::FocalLenght, 18.f)
+ *					.SetParameter(&ULensParametersCameraNode::FocalLength, 18.f)
  *					.Done()
  *				.Done()
  *			.AddEnterTransition()
@@ -491,11 +491,15 @@ public:
 		UCameraObjectInterfaceBlendableParameter* BlendableParameter = NewObject<UCameraObjectInterfaceBlendableParameter>(CameraRig);
 		BlendableParameter->InterfaceParameterName = ParameterName;
 		BlendableParameter->ParameterType = ParameterType;
-		BlendableParameter->Target = Target;
-		BlendableParameter->TargetPropertyName = TargetPropertyName;
 
 		NamedObjectRegistry->Register(BlendableParameter, ParameterName);
 		CameraRig->Interface.BlendableParameters.Add(BlendableParameter);
+
+		UCameraObjectInterfaceParameterGetter* Getter = NewObject<UCameraObjectInterfaceParameterGetter>(CameraRig);
+		Getter->ParameterGuid = BlendableParameter->GetGuid();
+		ensure(Getter->ParameterGuid.IsValid());
+		CameraRig->Connections.Add(Getter, NAME_None, Target, TargetPropertyName);
+
 		return *static_cast<ThisType*>(this);
 	}
 
diff --git a/Engine/Plugins/Cameras/GameplayCameras/Source/GameplayCameras/Public/Build/CameraAssetBuilder.h b/Engine/Plugins/Cameras/GameplayCameras/Source/GameplayCameras/Public/Build/CameraAssetBuilder.h
--- a/Engine/Plugins/Cameras/GameplayCameras/Source/GameplayCameras/Public/Build/CameraAssetBuilder.h
+++ b/Engine/Plugins/Cameras/GameplayCameras/Source/GameplayCameras/Public/Build/CameraAssetBuilder.h
@@ -9,6 +9,8 @@
 #define UE_API GAMEPLAYCAMERAS_API
 
 class UCameraAsset;
+class UCameraRigAsset;
+struct FInstancedPropertyBag;
 
 namespace UE::Cameras
 {
@@ -30,6 +32,8 @@ private:
 
 	void BuildCameraImpl(bool bBuildReferencedAssets);
 
+	void CopyDefaultParameterValues(const UCameraRigAsset* InCameraRig, FInstancedPropertyBag& DefaultParameters);
+
 	void UpdateBuildStatus();
 
 private:
diff --git a/Engine/Plugins/Cameras/GameplayCameras/Source/GameplayCameras/Public/Build/CameraObjectInterfaceBuilder.h b/Engine/Plugins/Cameras/GameplayCameras/Source/GameplayCameras/Public/Build/CameraObjectInterfaceBuilder.h
--- a/Engine/Plugins/Cameras/GameplayCameras/Source/GameplayCameras/Public/Build/CameraObjectInterfaceBuilder.h
+++ b/Engine/Plugins/Cameras/GameplayCameras/Source/GameplayCameras/Public/Build/CameraObjectInterfaceBuilder.h
@@ -42,11 +42,11 @@ private:
 
 private:
 
-	UE_API bool SetupCameraParameterOrVariableReferenceOverride(const UCameraObjectInterfaceBlendableParameter* BlendableParameter);
-	UE_API bool SetupCustomBlendableParameterOverride(const UCameraObjectInterfaceBlendableParameter* BlendableParameter);
+	UE_API bool SetupCameraParameterOrVariableReferenceOverride(const UCameraObjectInterfaceBlendableParameter* BlendableParameter, UCameraNode* Target, FName TargetPropertyName);
+	UE_API bool SetupCustomBlendableParameterOverride(const UCameraObjectInterfaceBlendableParameter* BlendableParameter, UCameraNode* Target, FName TargetPropertyName);
 
-	UE_API bool SetupDataContextPropertyOverride(const UCameraObjectInterfaceDataParameter* DataParameter);
-	UE_API bool SetupCustomDataParameterOverride(const UCameraObjectInterfaceDataParameter* DataParameter);
+	UE_API bool SetupDataContextPropertyOverride(const UCameraObjectInterfaceDataParameter* DataParameter, UCameraNode* Target, FName TargetPropertyName);
+	UE_API bool SetupCustomDataParameterOverride(const UCameraObjectInterfaceDataParameter* DataParameter, UCameraNode* Target, FName TargetPropertyName);
 
 private:
 
diff --git a/Engine/Plugins/Cameras/GameplayCameras/Source/GameplayCameras/Public/Build/CameraObjectInterfaceParameterBuilder.h b/Engine/Plugins/Cameras/GameplayCameras/Source/GameplayCameras/Public/Build/CameraObjectInterfaceParameterBuilder.h
--- a/Engine/Plugins/Cameras/GameplayCameras/Source/GameplayCameras/Public/Build/CameraObjectInterfaceParameterBuilder.h
+++ b/Engine/Plugins/Cameras/GameplayCameras/Source/GameplayCameras/Public/Build/CameraObjectInterfaceParameterBuilder.h
@@ -8,6 +8,7 @@
 #define UE_API GAMEPLAYCAMERAS_API
 
 class UBaseCameraObject;
+class UCameraNode;
 class UCameraObjectInterfaceDataParameter;
 struct FCameraObjectInterfaceParameterDefinition;
 struct FInstancedPropertyBag;
@@ -16,12 +17,6 @@ struct FPropertyBagPropertyDesc;
 namespace UE::Cameras
 {
 
-
-/**
- * A helper class for building an FInstancedPropertyBag from a list of camera rig
- * parameter definitions.
- */
-/** Builds the parameter definitions for the given camera object. */
 /**
  * Builds the property bag that contains a property for each exposed parameter on the given camera object.
  * Each property's value is set to the default value of the corresponding parameter.
@@ -36,10 +31,10 @@ public:
 
 public:
 
-	static UE_API void BuildDefaultParameters(const UBaseCameraObject* CameraObject, FInstancedPropertyBag& OutPropertyBag);
 	static UE_API void AppendDefaultParameterProperties(const UBaseCameraObject* CameraObject, TArray<FPropertyBagPropertyDesc>& OutProperties);
 	static UE_API void AppendDefaultParameterProperties(TConstArrayView<FCameraObjectInterfaceParameterDefinition> ParameterDefinitions, TArray<FPropertyBagPropertyDesc>& OutProperties);
-	static UE_API void SetDefaultParameterValues(const UBaseCameraObject* CameraObject, FInstancedPropertyBag& PropertyBag);
+
+	static UE_API bool SetDefaultParameterValue(UBaseCameraObject* CameraObject, const FCameraObjectInterfaceParameterDefinition& ParameterDefinition, UCameraNode* TargetNode, FName TargetPropertyName, bool bAddParameterIfMissing = true);
 
 #if UE_VERSION_OLDER_THAN(5,8,0)
 	static void FixUpDefaultParameterProperties(TConstArrayView<FCameraObjectInterfaceParameterDefinition> ParameterDefinitions, FInstancedPropertyBag& InOutPropertyBag);
@@ -47,8 +42,6 @@ public:
 
 private:
 
-	static void SetDefaultParameterValue(const UCameraObjectInterfaceDataParameter* DataParameter, void* DestValuePtr, const void* SrcValuePtr);
-
 	void BuildParametersImpl();
 
 	void BuildParameterDefinitions();
diff --git a/Engine/Plugins/Cameras/GameplayCameras/Source/GameplayCameras/Public/Core/BaseCameraObject.h b/Engine/Plugins/Cameras/GameplayCameras/Source/GameplayCameras/Public/Core/BaseCameraObject.h
--- a/Engine/Plugins/Cameras/GameplayCameras/Source/GameplayCameras/Public/Core/BaseCameraObject.h
+++ b/Engine/Plugins/Cameras/GameplayCameras/Source/GameplayCameras/Public/Core/BaseCameraObject.h
@@ -15,6 +15,9 @@
 
 #include "BaseCameraObject.generated.h"
 
+#define UE_API GAMEPLAYCAMERAS_API
+
+class IObjectTreeGraphRootObject;
 class UCameraNode;
 
 namespace UE::Cameras
@@ -72,6 +75,61 @@ struct TStructOpsTypeTraits<FCameraObjectAllocationInfo> : public TStructOpsType
 	};
 };
 
+/**
+ * A connection inside a camera object graph.
+ */
+USTRUCT()
+struct FCameraObjectConnection
+{
+	GENERATED_BODY()
+	
+public:
+
+	/** The source object of the connection. */
+	UPROPERTY()
+	TObjectPtr<UObject> Source;
+
+	/** The source object property of the connection. */
+	UPROPERTY()
+	FName SourcePropertyName;
+
+	/** The target object of the connection. */
+	UPROPERTY()
+	TObjectPtr<UObject> Target;
+
+	/** The target object property of the connection. */
+	UPROPERTY()
+	FName TargetPropertyName;
+};
+
+/**
+ * A collection of connections inside a camera object graph.
+ */
+USTRUCT()
+struct FCameraObjectConnections
+{
+	GENERATED_BODY()
+
+public:
+
+	/** The list of connections. */
+	UPROPERTY()
+	TArray<FCameraObjectConnection> Connections;
+
+public:
+
+	UE_API void Add(UObject* InSource, FName InSourcePropertyName, UObject* InTarget, FName InTargetPropertyName);
+
+	UE_API FCameraObjectConnection* FindBySource(UObject* InSource);
+	UE_API FCameraObjectConnection* FindBySource(UObject* InSource, FName InSourcePropertyName);
+	UE_API FCameraObjectConnection* FindByTarget(UObject* InTarget);
+	UE_API FCameraObjectConnection* FindByTarget(UObject* InTarget, FName InTargetPropertyName);
+};
+
+/**
+ * A base class for a camera object that has a graph of camera nodes, connections between them,
+ * and some exposed parameters.
+ */
 UCLASS(Abstract, MinimalAPI)
 class UBaseCameraObject : public UObject
 {
@@ -83,6 +141,9 @@ public:
 	UPROPERTY()
 	FCameraObjectInterface Interface;
 
+	UPROPERTY()
+	FCameraObjectConnections Connections;
+
 	/** Event handlers to be notified of data changes. */
 	UE::Cameras::TCameraEventHandlerContainer<UE::Cameras::ICameraObjectEventHandler> EventHandlers;
 
@@ -116,6 +177,13 @@ public:
 	virtual void PostInitProperties() override;
 	virtual void PostDuplicate(EDuplicateMode::Type DuplicateMode) override;
 
+protected:
+
+#if WITH_EDITORONLY_DATA
+	// Utility method to be called on PostLoad by subclasses that need upgrading old data.
+	void UpgradeInterfaceConnections(IObjectTreeGraphRootObject* RootObject, FName DefaultGraphName);
+#endif
+
 private:
 
 	/** The camera object's unique ID. */
@@ -133,3 +201,5 @@ private:
 	friend class UE::Cameras::FCameraObjectInterfaceParameterBuilder;
 };
 
+#undef UE_API
+
diff --git a/Engine/Plugins/Cameras/GameplayCameras/Source/GameplayCameras/Public/Core/CameraObjectInterface.h b/Engine/Plugins/Cameras/GameplayCameras/Source/GameplayCameras/Public/Core/CameraObjectInterface.h
--- a/Engine/Plugins/Cameras/GameplayCameras/Source/GameplayCameras/Public/Core/CameraObjectInterface.h
+++ b/Engine/Plugins/Cameras/GameplayCameras/Source/GameplayCameras/Public/Core/CameraObjectInterface.h
@@ -10,52 +10,40 @@
 
 #include "CameraObjectInterface.generated.h"
 
+#define UE_API GAMEPLAYCAMERAS_API
+
 class UCameraNode;
 struct FCameraContextDataDefinition;
+struct FCameraObjectInterfaceParameterDefinition;
 struct FCameraVariableDefinition;
 
 /**
  * Base class for interface parameters on a camera rig asset.
  */
-UCLASS(MinimalAPI, meta=(ObjectTreeGraphSelfPinDirection="Output"))
+UCLASS(MinimalAPI)
 class UCameraObjectInterfaceParameterBase 
 	: public UObject
-	, public IObjectTreeGraphObject
 {
 	GENERATED_BODY()
 
 public:
 
 	/** The exposed name for this parameter. */
-	UPROPERTY(EditAnywhere, Category=Camera)
+	UPROPERTY(EditAnywhere, Category="Camera")
 	FString InterfaceParameterName;
 
-	/** The camera node this parameter is connected to. */
-	UPROPERTY(meta=(ObjectTreeGraphHidden=true))
-	TObjectPtr<UCameraNode> Target;
-
-	/**
-	 * The name of the property this parameter is connected to on the target camera node.
-	 * This may be an actual UObject property, but it may be something else, like the name
-	 * of an interface parameter on a nested camera rig, or the name of a Blueprint property
-	 * on the evaluator class of a Blueprint camera node.
-	 */
-	UPROPERTY()
-	FName TargetPropertyName;
-
-#if WITH_EDITORONLY_DATA
-
-	/** Whether this parameter has been added to the node graph in the editor. */
-	UPROPERTY()
-	bool bHasGraphNode = false;
-
-#endif  // WITH_EDITORONLY_DATA
+	/** Whether to show this parameter on prefab nodes, references, components, etc. */
+	UPROPERTY(EditAnywhere, Category="Visibility")
+	bool bIsVisible = true;
 
 public:
 
 	/** Gets this parameter's unique ID. */
 	const FGuid& GetGuid() const { return Guid; }
 
+	/** Gets the parameter definition. */
+	virtual void GetParameterDefinition(FCameraObjectInterfaceParameterDefinition& OutParameterDefinition) const {}
+
 protected:
 
 	/** The Guid of this parameter. */
@@ -64,23 +52,28 @@ protected:
 
 protected:
 
-	// IObjectTreeGraphObject interface.
-#if WITH_EDITOR
-	virtual void GetGraphNodePosition(FName InGraphName, int32& NodePosX, int32& NodePosY) const override;
-	virtual void OnGraphNodeMoved(FName InGraphName, int32 NodePosX, int32 NodePosY, bool bMarkDirty) override;
-#endif
-
 	// UObject interface.
 	virtual void PostLoad() override;
 	virtual void PostInitProperties() override;
 	virtual void PostDuplicate(EDuplicateMode::Type DuplicateMode) override;
 
-private:
+public:
+
+	// Deprecated
+
+	UPROPERTY(meta=(ObjectTreeGraphHidden=true))
+	TObjectPtr<UCameraNode> Target_DEPRECATED;
+
+	UPROPERTY()
+	FName TargetPropertyName_DEPRECATED;
 
 #if WITH_EDITORONLY_DATA
 
 	UPROPERTY()
-	FIntVector2 GraphNodePos = FIntVector2::ZeroValue;
+	FIntVector2 GraphNodePos_DEPRECATED = FIntVector2::ZeroValue;
+
+	UPROPERTY()
+	bool bHasGraphNode_DEPRECATED = false;
 
 #endif  // WITH_EDITORONLY_DATA
 };
@@ -111,7 +104,7 @@ public:
 	 * each of their values will be blended in a first evaluation pass, and then
 	 * both camera rigs will evaluate with the same blended value.
 	 */
-	UPROPERTY()
+	UPROPERTY(EditAnywhere, Category="Blending")
 	bool bIsPreBlended = false;
 
 	// Built on save/cook.
@@ -134,6 +127,9 @@ public:
 #if WITH_EDITORONLY_DATA
 	GAMEPLAYCAMERAS_API FString GetVariableName() const;
 #endif
+
+	// UCameraObjectInterfaceParameterBase interface.
+	virtual void GetParameterDefinition(FCameraObjectInterfaceParameterDefinition& OutParameterDefinition) const override;
 };
 
 UCLASS(MinimalAPI)
@@ -169,6 +165,9 @@ public:
 #if WITH_EDITORONLY_DATA
 	GAMEPLAYCAMERAS_API FString GetDataName() const;
 #endif
+
+	// UCameraObjectInterfaceParameterBase interface.
+	virtual void GetParameterDefinition(FCameraObjectInterfaceParameterDefinition& OutParameterDefinition) const override;
 };
 
 /**
@@ -216,3 +215,40 @@ private:
 	FString DisplayName_DEPRECATED;
 };
 
+/**
+ * Getter object for adding a camera parameter to a camera node graph.
+ */
+UCLASS(MinimalAPI, meta=(ObjectTreeGraphSelfPinDirection="Output"))
+class UCameraObjectInterfaceParameterGetter
+	: public UObject
+	, public IObjectTreeGraphObject
+{
+	GENERATED_BODY()
+
+public:
+
+	/** The Guid of the parameter to get. */
+	UPROPERTY()
+	FGuid ParameterGuid;
+
+	/** The location of this node in the graph. */
+	UPROPERTY()
+	FIntVector2 GraphNodePos = FIntVector2::ZeroValue;
+
+public:
+
+	UE_API UCameraObjectInterfaceParameterBase* GetInterfaceParameter() const;
+
+	UE_API FString GetInterfaceParameterName() const;
+	
+protected:
+
+	// IObjectTreeGraphObject interface.
+#if WITH_EDITOR
+	virtual void GetGraphNodePosition(FName InGraphName, int32& NodePosX, int32& NodePosY) const override;
+	virtual void OnGraphNodeMoved(FName InGraphName, int32 NodePosX, int32 NodePosY, bool bMarkDirty) override;
+#endif
+};
+
+#undef UE_API
+
diff --git a/Engine/Plugins/Cameras/GameplayCameras/Source/GameplayCameras/Public/Core/CameraObjectInterfaceParameterDefinition.h b/Engine/Plugins/Cameras/GameplayCameras/Source/GameplayCameras/Public/Core/CameraObjectInterfaceParameterDefinition.h
--- a/Engine/Plugins/Cameras/GameplayCameras/Source/GameplayCameras/Public/Core/CameraObjectInterfaceParameterDefinition.h
+++ b/Engine/Plugins/Cameras/GameplayCameras/Source/GameplayCameras/Public/Core/CameraObjectInterfaceParameterDefinition.h
@@ -8,7 +8,7 @@
 #include "CameraObjectInterfaceParameterDefinition.generated.h"
 
 /**
- * The type of a camera rig parameter.
+ * The type of a camera parameter.
  */
 UENUM()
 enum class ECameraObjectInterfaceParameterType : uint8
@@ -41,6 +41,10 @@ struct FCameraObjectInterfaceParameterDefinition
 	UPROPERTY()
 	ECameraObjectInterfaceParameterType ParameterType = ECameraObjectInterfaceParameterType::Blendable;
 
+	/** Whether to show this parameter on prefab nodes, references, components, etc. */
+	UPROPERTY()
+	bool bIsVisible = true;
+
 
 	// Blendable parameter properties.
 	// (only valid when ParameterType == Blendable)
diff --git a/Engine/Plugins/Cameras/GameplayCameras/Source/GameplayCameras/Public/Core/CameraShakeAsset.h b/Engine/Plugins/Cameras/GameplayCameras/Source/GameplayCameras/Public/Core/CameraShakeAsset.h
--- a/Engine/Plugins/Cameras/GameplayCameras/Source/GameplayCameras/Public/Core/CameraShakeAsset.h
+++ b/Engine/Plugins/Cameras/GameplayCameras/Source/GameplayCameras/Public/Core/CameraShakeAsset.h
@@ -90,6 +90,7 @@ protected:
 #endif
 
 	// UObject interface.
+	virtual void PostLoad() override;
 	virtual void PreSave(FObjectPreSaveContext ObjectSaveContext) override;
 
 private:
diff --git a/Engine/Plugins/Cameras/GameplayCameras/Source/GameplayCamerasEditor/Private/Customizations/CameraObjectInterfaceParameterDetailsCustomization.cpp b/Engine/Plugins/Cameras/GameplayCameras/Source/GameplayCamerasEditor/Private/Customizations/CameraObjectInterfaceParameterDetailsCustomization.cpp
new file mode 100644
--- /dev/null
+++ b/Engine/Plugins/Cameras/GameplayCameras/Source/GameplayCamerasEditor/Private/Customizations/CameraObjectInterfaceParameterDetailsCustomization.cpp
@@ -0,0 +1,110 @@
+// Copyright Epic Games, Inc. All Rights Reserved.
+
+#include "Customizations/CameraObjectInterfaceParameterDetailsCustomization.h"
+
+#include "Core/BaseCameraObject.h"
+#include "Core/CameraObjectInterface.h"
+#include "DetailCategoryBuilder.h"
+#include "DetailLayoutBuilder.h"
+#include "DetailWidgetRow.h"
+#include "IDetailChildrenBuilder.h"
+#include "IDetailGroup.h"
+#include "PropertyCustomizationHelpers.h"
+#include "Widgets/Text/STextBlock.h"
+
+#define LOCTEXT_NAMESPACE "CameraObjectInterfaceParameterDetailsCustomization"
+
+namespace UE::Cameras
+{
+
+TSharedRef<IDetailCustomization> FCameraObjectInterfaceParameterDetailsCustomization::MakeInstance()
+{
+	return MakeShared<FCameraObjectInterfaceParameterDetailsCustomization>();
+}
+
+void FCameraObjectInterfaceParameterDetailsCustomization::CustomizeDetails(IDetailLayoutBuilder& DetailBuilder)
+{
+	bool bFoundValidOuter = false;
+	UBaseCameraObject* CommonCameraObject = nullptr;
+	TSet<FGuid> ParameterGuids;
+
+	TArray<TWeakObjectPtr<UObject>> WeakObjects;
+	DetailBuilder.GetObjectsBeingCustomized(WeakObjects);
+	for (TWeakObjectPtr<UObject> WeakObj : WeakObjects)
+	{
+		if (UCameraObjectInterfaceParameterBase* Obj = Cast<UCameraObjectInterfaceParameterBase>(WeakObj.Get()))
+		{
+			ParameterGuids.Add(Obj->GetGuid());
+
+			UBaseCameraObject* CameraObject = Obj->GetTypedOuter<UBaseCameraObject>();
+			if (!bFoundValidOuter && CameraObject)
+			{
+				CommonCameraObject = CameraObject;
+				bFoundValidOuter = true;
+			}
+			else if (bFoundValidOuter && CommonCameraObject && CameraObject != CommonCameraObject)
+			{
+				CommonCameraObject = nullptr;
+			}
+		}
+	}
+	if (!CommonCameraObject)
+	{
+		return;
+	}
+
+	TArray<FName> PropertyNames;
+	FInstancedPropertyBag& DefaultParameters = CommonCameraObject->GetDefaultParameters();
+	const UPropertyBag* DefaultParametersStruct = DefaultParameters.GetPropertyBagStruct();
+	if (DefaultParametersStruct)
+	{
+		for (const FPropertyBagPropertyDesc& Item : DefaultParametersStruct->GetPropertyDescs())
+		{
+			if (ParameterGuids.Contains(Item.ID))
+			{
+				PropertyNames.Add(Item.Name);
+			}
+		}
+	}
+
+	IDetailCategoryBuilder& CategoryBuilder = DetailBuilder.EditCategory(
+			TEXT("DefaultValue"), 
+			LOCTEXT("DefaultValueCategory", "Default Value"));
+
+	// WARNING: here we put a Struct-On-Scope inside the details view, and it has a direct access to the default 
+	// parameters' memory. However, that memory may move when we build the current asset if we have added enough 
+	// new parameters to require re-allocation. The Details View would then suddenly start hitting freed memory.
+	// We fix this by forcing a Details View refresh on build in all the toolkits that do it.
+	if (DefaultParametersStruct)
+	{
+		auto ModifyCommonCameraObject = [CommonCameraObject]() { CommonCameraObject->Modify(); };
+		TSharedRef<FStructOnScope> StructOnScope = MakeShared<FStructOnScope>(DefaultParametersStruct, DefaultParameters.GetMutableValue().GetMemory());
+		for (FName PropertyName : PropertyNames)
+		{
+			FAddPropertyParams AddParams;
+			AddParams.ForceShowProperty();  // The property may not have CPF_Edit if the parameter is set to not-visible,
+											// so force it to be shown.
+			IDetailPropertyRow* PropertyRow = CategoryBuilder.AddExternalStructureProperty(StructOnScope, PropertyName, EPropertyLocation::Default, AddParams);
+
+			TSharedPtr<IPropertyHandle> DefaultValuePropertyHandle = PropertyRow->GetPropertyHandle();
+			DefaultValuePropertyHandle->SetOnPropertyValuePreChange(FSimpleDelegate::CreateLambda(ModifyCommonCameraObject));
+			DefaultValuePropertyHandle->SetOnChildPropertyValuePreChange(FSimpleDelegate::CreateLambda(ModifyCommonCameraObject));
+			DefaultValuePropertyHandle->SetOnPropertyResetToDefault(FSimpleDelegate::CreateLambda(ModifyCommonCameraObject));
+
+		}
+	}
+	if (PropertyNames.Num() < ParameterGuids.Num())
+	{
+		FDetailWidgetRow& CustomRow = CategoryBuilder.AddCustomRow(LOCTEXT("UnbuiltDataRow", "Needs Building"));
+		CustomRow.WholeRowContent()
+			[
+				SNew(STextBlock)
+				.Text(LOCTEXT("NeedsBuildWarning", "Please build the asset to edit default parameter values"))
+			];
+	}
+}
+
+}  // namespace UE::Cameras
+
+#undef LOCTEXT_NAMESPACE
+
diff --git a/Engine/Plugins/Cameras/GameplayCameras/Source/GameplayCamerasEditor/Private/Customizations/CameraParameterDetailsCustomizations.cpp b/Engine/Plugins/Cameras/GameplayCameras/Source/GameplayCamerasEditor/Private/Customizations/CameraParameterDetailsCustomizations.cpp
--- a/Engine/Plugins/Cameras/GameplayCameras/Source/GameplayCamerasEditor/Private/Customizations/CameraParameterDetailsCustomizations.cpp
+++ b/Engine/Plugins/Cameras/GameplayCameras/Source/GameplayCamerasEditor/Private/Customizations/CameraParameterDetailsCustomizations.cpp
@@ -211,6 +211,7 @@ void FCameraParameterDetailsCustomization::UpdateVariableInfo()
 					if (RawData)
 					{
 						VariableInfo.bHasNonUserOverride |= HasNonUserOverride(RawData);
+						VariableInfo.bHasOverride |= HasOverride(RawData);
 					}
 					return true;
 				});
@@ -305,8 +306,8 @@ TSharedRef<SWidget> FCameraParameterDetailsCustomization::BuildCameraVariableBro
 bool FCameraParameterDetailsCustomization::IsValueEditorEnabled() const
 {
 	// The value widget is enabled (i.e. the user can change the value) if the parameter isn't driven by
-	// a variable that was set by the user.
-	return VariableInfo.VariableValue == ECameraVariableValue::NotSet;
+	// a variable or a rig parameter;
+	return !VariableInfo.bHasOverride;
 }
 
 bool FCameraParameterDetailsCustomization::IsCameraVariableBrowserEnabled() const
@@ -464,6 +465,11 @@ void FCameraParameterDetailsCustomization::OnResetToDefault(TSharedPtr<IProperty
 }
 
 #define UE_CAMERA_VARIABLE_FOR_TYPE(ValueType, ValueName)\
+bool F##ValueName##CameraParameterDetailsCustomization::HasOverride(void* InRawData)\
+{\
+	F##ValueName##CameraParameter* TypedData = reinterpret_cast<F##ValueName##CameraParameter*>(InRawData);\
+	return TypedData->HasOverride();\
+}\
 bool F##ValueName##CameraParameterDetailsCustomization::HasNonUserOverride(void* InRawData)\
 {\
 	F##ValueName##CameraParameter* TypedData = reinterpret_cast<F##ValueName##CameraParameter*>(InRawData);\
diff --git a/Engine/Plugins/Cameras/GameplayCameras/Source/GameplayCamerasEditor/Private/Editors/CameraNodeGraphDragDropOp.cpp b/Engine/Plugins/Cameras/GameplayCameras/Source/GameplayCamerasEditor/Private/Editors/CameraNodeGraphDragDropOp.cpp
--- a/Engine/Plugins/Cameras/GameplayCameras/Source/GameplayCamerasEditor/Private/Editors/CameraNodeGraphDragDropOp.cpp
+++ b/Engine/Plugins/Cameras/GameplayCameras/Source/GameplayCamerasEditor/Private/Editors/CameraNodeGraphDragDropOp.cpp
@@ -3,7 +3,7 @@
 #include "Editors/CameraNodeGraphDragDropOp.h"
 
 #include "Core/CameraRigAsset.h"
-#include "Editors/CameraNodeGraphSchema.h"
+#include "Editors/CameraObjectGraphSchemaBase.h"
 #include "GraphEditor.h"
 #include "ScopedTransaction.h"
 
@@ -21,18 +21,8 @@ FReply FCameraNodeGraphInterfaceParameterDragDropOp::ExecuteDragOver(TSharedPtr<
 {
 	if (InterfaceParameter)
 	{
-		if (!InterfaceParameter->bHasGraphNode)
-		{
-			const FSlateBrush* OKIcon = FAppStyle::GetBrush(TEXT("Graph.ConnectorFeedback.OK"));
-			SetToolTip(LOCTEXT("OnDragOver_Success", "Add interface parameter"), OKIcon);
-		}
-		else
-		{
-			const FSlateBrush* ErrorIcon = FAppStyle::GetBrush(TEXT("Graph.ConnectorFeedback.Error"));
-			SetToolTip(
-					LOCTEXT("OnDragOver_Error", "This interface parameter is already in the graph"),
-					ErrorIcon);
-		}
+		const FSlateBrush* OKIcon = FAppStyle::GetBrush(TEXT("Graph.ConnectorFeedback.OK"));
+		SetToolTip(LOCTEXT("OnDragOver_Success", "Add interface parameter getter"), OKIcon);
 	}
 
 	return FReply::Handled();
@@ -40,7 +30,7 @@ FReply FCameraNodeGraphInterfaceParameterDragDropOp::ExecuteDragOver(TSharedPtr<
 
 FReply FCameraNodeGraphInterfaceParameterDragDropOp::ExecuteDrop(TSharedPtr<SGraphEditor> GraphEditor, const FSlateCompatVector2f& NewLocation)
 {
-	if (!InterfaceParameter || InterfaceParameter->bHasGraphNode)
+	if (!InterfaceParameter)
 	{
 		return FReply::Handled();
 	}
@@ -51,7 +41,7 @@ FReply FCameraNodeGraphInterfaceParameterDragDropOp::ExecuteDrop(TSharedPtr<SGra
 
 	GraphEditor->ClearSelectionSet();
 
-	FCameraNodeGraphSchemaAction_AddInterfaceParameterNode Action;
+	FCameraObjectGraphSchemaAction_AddInterfaceParameterGetterNode Action;
 	Action.InterfaceParameter = InterfaceParameter;
 	UEdGraphNode* NewNode = Action.PerformAction(Graph, nullptr, NewLocation, false);
 	GraphEditor->SetNodeSelection(NewNode, true);
diff --git a/Engine/Plugins/Cameras/GameplayCameras/Source/GameplayCamerasEditor/Private/Editors/CameraNodeGraphNode.cpp b/Engine/Plugins/Cameras/GameplayCameras/Source/GameplayCamerasEditor/Private/Editors/CameraNodeGraphNode.cpp
--- a/Engine/Plugins/Cameras/GameplayCameras/Source/GameplayCamerasEditor/Private/Editors/CameraNodeGraphNode.cpp
+++ b/Engine/Plugins/Cameras/GameplayCameras/Source/GameplayCamerasEditor/Private/Editors/CameraNodeGraphNode.cpp
@@ -8,7 +8,7 @@
 #include "Core/CameraVariableReferences.h"
 #include "Core/ICustomCameraNodeParameterProvider.h"
 #include "EdGraph/EdGraphPin.h"
-#include "Editors/CameraNodeGraphSchema.h"
+#include "Editors/CameraObjectGraphSchemaBase.h"
 #include "Editors/SCameraNodeGraphNode.h"
 #include "GameplayCamerasDelegates.h"
 #include "ToolMenus.h"
@@ -71,7 +71,7 @@ void UCameraNodeGraphNode::AllocateDefaultPins()
 #define UE_CAMERA_VARIABLE_FOR_TYPE(ValueType, ValueName)\
 			if (StructProperty->Struct == F##ValueName##CameraParameter::StaticStruct())\
 			{\
-				PinType.PinCategory = UCameraNodeGraphSchema::PC_CameraParameter;\
+				PinType.PinCategory = UCameraObjectGraphSchemaBase::PC_CameraParameter;\
 				PinType.PinSubCategory = UEnum::GetValueAsName(ECameraVariableType::ValueName);\
 				PinType.PinSubCategoryObject = F##ValueName##CameraParameter::StaticStruct();\
 				UEdGraphPin* ParameterPin = CreatePin(EGPD_Input, PinType, PropertyName);\
@@ -81,7 +81,7 @@ void UCameraNodeGraphNode::AllocateDefaultPins()
 			}\
 			if (StructProperty->Struct == F##ValueName##CameraVariableReference::StaticStruct())\
 			{\
-				PinType.PinCategory = UCameraNodeGraphSchema::PC_CameraVariableReference;\
+				PinType.PinCategory = UCameraObjectGraphSchemaBase::PC_CameraVariableReference;\
 				PinType.PinSubCategory = UEnum::GetValueAsName(ECameraVariableType::ValueName);\
 				PinType.PinSubCategoryObject = F##ValueName##CameraVariableReference::StaticStruct();\
 				UEdGraphPin* VariableReferencePin = CreatePin(EGPD_Input, PinType, PropertyName);\
@@ -142,7 +142,7 @@ void UCameraNodeGraphNode::AllocateDefaultPins()
 
 			if (bGotValidDataProperty)
 			{
-				PinType.PinCategory = UCameraNodeGraphSchema::PC_CameraContextData;
+				PinType.PinCategory = UCameraObjectGraphSchemaBase::PC_CameraContextData;
 				UEdGraphPin* ContextDataPin = CreatePin(EGPD_Input, PinType, PropertyName);
 				ContextDataPin->PinFriendlyName = PinFriendlyName;
 				ContextDataPin->PinToolTip = PinToolTip;
@@ -165,7 +165,7 @@ void UCameraNodeGraphNode::AllocateDefaultPins()
 		for (const FCustomCameraNodeBlendableParameter& BlendableParameter : CustomBlendableParameters)
 		{
 			FEdGraphPinType PinType;
-			PinType.PinCategory = UCameraNodeGraphSchema::PC_CameraParameter;
+			PinType.PinCategory = UCameraObjectGraphSchemaBase::PC_CameraParameter;
 			PinType.PinSubCategory = VariableTypeEnum->GetValueAsName(BlendableParameter.ParameterType);
 
 			switch (BlendableParameter.ParameterType)
@@ -194,7 +194,7 @@ void UCameraNodeGraphNode::AllocateDefaultPins()
 		for (const FCustomCameraNodeDataParameter& DataParameter : CustomDataParameters)
 		{
 			FEdGraphPinType PinType;
-			PinType.PinCategory = UCameraNodeGraphSchema::PC_CameraContextData;
+			PinType.PinCategory = UCameraObjectGraphSchemaBase::PC_CameraContextData;
 			PinType.PinSubCategory = DataTypeEnum->GetNameByValue((int64)DataParameter.ParameterType);
 			PinType.PinSubCategoryObject = const_cast<UObject*>(DataParameter.ParameterTypeObject.Get());
 
diff --git a/Engine/Plugins/Cameras/GameplayCameras/Source/GameplayCamerasEditor/Private/Editors/CameraNodeGraphSchema.cpp b/Engine/Plugins/Cameras/GameplayCameras/Source/GameplayCamerasEditor/Private/Editors/CameraNodeGraphSchema.cpp
deleted file mode 100644
diff --git a/Engine/Plugins/Cameras/GameplayCameras/Source/GameplayCamerasEditor/Private/Editors/CameraObjectGraphSchemaBase.cpp b/Engine/Plugins/Cameras/GameplayCameras/Source/GameplayCamerasEditor/Private/Editors/CameraObjectGraphSchemaBase.cpp
new file mode 100644
--- /dev/null
+++ b/Engine/Plugins/Cameras/GameplayCameras/Source/GameplayCamerasEditor/Private/Editors/CameraObjectGraphSchemaBase.cpp
@@ -0,0 +1,693 @@
+// Copyright Epic Games, Inc. All Rights Reserved.
+
+#include "Editors/CameraObjectGraphSchemaBase.h"
+
+#include "Algo/Reverse.h"
+#include "Build/CameraObjectInterfaceParameterBuilder.h"
+#include "Core/BaseCameraObject.h"
+#include "Core/CameraNode.h"
+#include "Core/CameraParameters.h"
+#include "Core/CameraVariableReferences.h"
+#include "Core/ICustomCameraNodeParameterProvider.h"
+#include "EdGraph/EdGraphPin.h"
+#include "EdGraphSchema_K2.h"
+#include "Editors/CameraNodeGraphNode.h"
+#include "Editors/CameraObjectInterfaceParameterGraphNode.h"
+#include "Editors/ObjectTreeGraph.h"
+#include "Editors/ObjectTreeGraphConfig.h"
+#include "Editors/ObjectTreeGraphNode.h"
+#include "GameplayCamerasEditorSettings.h"
+#include "ScopedTransaction.h"
+
+#include UE_INLINE_GENERATED_CPP_BY_NAME(CameraObjectGraphSchemaBase)
+
+#define LOCTEXT_NAMESPACE "CameraObjectGraphSchemaBase"
+
+const FName UCameraObjectGraphSchemaBase::PC_CameraParameter("CameraParameter");
+const FName UCameraObjectGraphSchemaBase::PC_CameraVariableReference("CameraVariableReference");
+const FName UCameraObjectGraphSchemaBase::PC_CameraContextData("CameraContextData");
+
+UCameraObjectGraphSchemaBase::UCameraObjectGraphSchemaBase(const FObjectInitializer& ObjInit)
+	: Super(ObjInit)
+{
+	PinColors.Initialize();
+}
+
+FObjectTreeGraphConfig UCameraObjectGraphSchemaBase::BuildGraphConfig() const
+{
+	const UGameplayCamerasEditorSettings* Settings = GetDefault<UGameplayCamerasEditorSettings>();
+
+	FObjectTreeGraphConfig GraphConfig;
+
+	GraphConfig.DefaultSelfPinName = NAME_None;
+	GraphConfig.ConnectableObjectClasses.Add(UCameraObjectInterfaceParameterGetter::StaticClass());
+	GraphConfig.ObjectClassConfigs.Emplace(UCameraObjectInterfaceParameterGetter::StaticClass())
+		.GraphNodeClass(UCameraObjectInterfaceParameterGraphNode::StaticClass());
+
+	OnBuildGraphConfig(GraphConfig);
+
+	return GraphConfig;
+}
+
+void UCameraObjectGraphSchemaBase::OnCreateAllNodes(UObjectTreeGraph* InGraph, const FCreatedNodes& InCreatedNodes) const
+{
+	Super::OnCreateAllNodes(InGraph, InCreatedNodes);
+
+	// Once all nodes are created, add connections between them.
+	CreateValueFlowConnections(InGraph, InCreatedNodes);
+}
+
+void UCameraObjectGraphSchemaBase::CreateValueFlowConnections(UObjectTreeGraph* InGraph, const FCreatedNodes& InCreatedNodes) const
+{
+	UBaseCameraObject* CameraObject = Cast<UBaseCameraObject>(InGraph->GetRootObject());
+	if (!CameraObject)
+	{
+		return;
+	}
+
+	const FName SelfPinName = UObjectTreeGraphSchema::PC_Self;
+
+	// Add connections for connected objects that are in this graph.
+	for (const FCameraObjectConnection& Connection : CameraObject->Connections.Connections)
+	{
+		if (!Connection.Source || !Connection.Target)
+		{
+			continue;
+		}
+
+		UEdGraphNode* SourceNode = InCreatedNodes.CreatedNodes.FindRef(Connection.Source);
+		UEdGraphNode* TargetNode = InCreatedNodes.CreatedNodes.FindRef(Connection.Target);
+		// Either both nodes in the connection are in this graph, or none of them are. We shouldn't have
+		// one node from one graph connected to a node from another graph, or nodes belonging to two graphs.
+		ensure((SourceNode && TargetNode) || (!SourceNode && !TargetNode));
+		if (!SourceNode || !TargetNode)
+		{
+			continue;
+		}
+
+		UEdGraphPin* SourcePin = Connection.SourcePropertyName.IsNone() ?
+			FindPinByType(SourceNode, PC_Self) : 
+			SourceNode->FindPin(Connection.SourcePropertyName);
+		if (!SourcePin)
+		{
+			SourcePin = SourceNode->CreatePin(EGPD_Output, NAME_None, Connection.SourcePropertyName);
+			SourcePin->bOrphanedPin = true;
+		}
+
+		UEdGraphPin* TargetPin = Connection.TargetPropertyName.IsNone() ? 
+			FindPinByType(TargetNode, PC_Self) : 
+			TargetNode->FindPin(Connection.TargetPropertyName);
+		if (!TargetPin)
+		{
+			TargetPin = TargetNode->CreatePin(EGPD_Input, NAME_None, Connection.TargetPropertyName);
+			TargetPin->bOrphanedPin = true;
+		}
+
+		SourcePin->MakeLinkTo(TargetPin);
+	}
+}
+
+UEdGraphPin* UCameraObjectGraphSchemaBase::FindPinByType(UEdGraphNode* InNode, const FName& InPinCategory) const
+{
+	UEdGraphPin* const* FoundItem = InNode->Pins.FindByPredicate(
+			[InPinCategory](UEdGraphPin* Item)
+			{ 
+				return Item->PinType.PinCategory == InPinCategory;
+			});
+	if (FoundItem)
+	{
+		return *FoundItem;
+	}
+	return nullptr;
+}
+
+void UCameraObjectGraphSchemaBase::GetGraphContextActions(FGraphContextMenuBuilder& ContextMenuBuilder) const
+{
+	// See if we were dragging a camera parameter pin or camera variable reference pin.
+	if (const UEdGraphPin* DraggedPin = ContextMenuBuilder.FromPin)
+	{
+		UCameraNodeGraphNode* CameraNodeNode = Cast<UCameraNodeGraphNode>(DraggedPin->GetOwningNode());
+
+		if (DraggedPin->PinType.PinCategory == PC_CameraParameter ||
+				DraggedPin->PinType.PinCategory == PC_CameraVariableReference ||
+				DraggedPin->PinType.PinCategory == PC_CameraContextData)
+		{
+			ensure(DraggedPin->PinName != NAME_None);
+
+			// If this is an invalid parameter/data pin, don't show any actions.
+			if (DraggedPin->bOrphanedPin)
+			{
+				return;
+			}
+
+			// Find the property being dragged, so we know what kind of parameter to create.
+			const UClass* CameraNodeClass = CameraNodeNode->GetObject()->GetClass();
+			FProperty* Property = CameraNodeClass->FindPropertyByName(DraggedPin->PinName);
+
+			FCustomCameraNodeParameterInfos CustomParameters;
+			ICustomCameraNodeParameterProvider* CustomParameterProvider = Cast<ICustomCameraNodeParameterProvider>(CameraNodeNode->GetObject());
+			if (CustomParameterProvider)
+			{
+				CustomParameterProvider->GetCustomCameraNodeParameters(CustomParameters);
+			}
+
+			TSharedRef<FCameraObjectGraphSchemaAction_NewInterfaceParameterNode> Action = 
+				MakeShared<FCameraObjectGraphSchemaAction_NewInterfaceParameterNode>(
+						FText::GetEmpty(),
+						LOCTEXT("NewInterfaceParameterAction", "Camera Interface Parameter"),
+						LOCTEXT("NewInterfaceParameterActionToolTip", "Exposes this parameter on the camera object"));
+
+			if (DraggedPin->PinType.PinCategory == PC_CameraParameter ||
+				DraggedPin->PinType.PinCategory == PC_CameraVariableReference)
+			{
+				ECameraVariableType VariableType;
+				const UScriptStruct* BlendableStructType = nullptr;
+
+				FCustomCameraNodeBlendableParameter BlendableParameter;
+				if (CustomParameters.FindBlendableParameter(DraggedPin->PinName, BlendableParameter))
+				{
+					VariableType = BlendableParameter.ParameterType;
+					BlendableStructType = BlendableParameter.BlendableStructType;
+				}
+				else if (FStructProperty* StructProperty = CastField<FStructProperty>(Property))
+				{
+#define UE_CAMERA_VARIABLE_FOR_TYPE(ValueType, ValueName)\
+					if ((StructProperty->Struct == F##ValueName##CameraParameter::StaticStruct()) ||\
+						(StructProperty->Struct == F##ValueName##CameraVariableReference::StaticStruct()))\
+					{\
+						VariableType = ECameraVariableType::ValueName;\
+					}\
+					else
+					UE_CAMERA_VARIABLE_FOR_ALL_TYPES()
+#undef UE_CAMERA_VARIABLE_FOR_TYPE
+					{
+						// Unexpected: if there was a camera parameter pin or a variable reference pin, we should
+						// have had a camera parameter property or variable reference property!
+						ensure(false);
+						return;
+					}
+				}
+				else
+				{
+					// Unexpected as per previous comments.
+					ensure(false);
+					return;
+				}
+
+				FCameraObjectInterfaceParameterDefinition NewParameterDefinition;
+				NewParameterDefinition.ParameterType = ECameraObjectInterfaceParameterType::Blendable;
+				NewParameterDefinition.VariableType = VariableType;
+				NewParameterDefinition.BlendableStructType = BlendableStructType;
+				Action->ParameterDefinition = NewParameterDefinition;
+			}
+			else if (DraggedPin->PinType.PinCategory == PC_CameraContextData)
+			{
+				ECameraContextDataType DataType;
+				ECameraContextDataContainerType DataContainerType = ECameraContextDataContainerType::None;
+				const UObject* DataTypeObject = nullptr;
+				
+				FCustomCameraNodeDataParameter DataParameter;
+				if (CustomParameters.FindDataParameter(DraggedPin->PinName, DataParameter))
+				{
+					DataType = DataParameter.ParameterType;
+					DataContainerType = DataParameter.ParameterContainerType;
+					DataTypeObject = DataParameter.ParameterTypeObject;
+				}
+				else if (Property)
+				{
+					if (FArrayProperty* ArrayProperty = CastField<FArrayProperty>(Property))
+					{
+						DataContainerType = ECameraContextDataContainerType::Array;
+						Property = ArrayProperty->Inner;
+					}
+
+					if (FNameProperty* NameProperty = CastField<FNameProperty>(Property))
+					{
+						DataType = ECameraContextDataType::Name;
+					}
+					else if (FStrProperty* StringProperty = CastField<FStrProperty>(Property))
+					{
+						DataType = ECameraContextDataType::String;
+					}
+					else if (FEnumProperty* EnumProperty = CastField<FEnumProperty>(Property))
+					{
+						DataType = ECameraContextDataType::Enum;
+						DataTypeObject = EnumProperty->GetEnum();
+					}
+					else if (FStructProperty* StructProperty = CastField<FStructProperty>(Property))
+					{
+						DataType = ECameraContextDataType::Struct;
+						DataTypeObject = StructProperty->Struct;
+					}
+					else if (FClassProperty* ClassProperty = CastField<FClassProperty>(Property))
+					{
+						DataType = ECameraContextDataType::Class;
+						DataTypeObject = ClassProperty->PropertyClass;
+					}
+					else if (FObjectProperty* ObjectProperty = CastField<FObjectProperty>(Property))
+					{
+						DataType = ECameraContextDataType::Object;
+						DataTypeObject = ObjectProperty->PropertyClass;
+					}
+					else
+					{
+						// Unexpected as per previous comments.
+						ensure(false);
+						return;
+					}
+				}
+				else
+				{
+					// Unexpected as per previous comments.
+					ensure(false);
+					return;
+				}
+
+				FCameraObjectInterfaceParameterDefinition NewParameterDefinition;
+				NewParameterDefinition.ParameterType = ECameraObjectInterfaceParameterType::Data;
+				NewParameterDefinition.DataType = DataType;
+				NewParameterDefinition.DataContainerType = DataContainerType;
+				NewParameterDefinition.DataTypeObject = DataTypeObject;
+				Action->ParameterDefinition = NewParameterDefinition;
+			}
+
+			ContextMenuBuilder.AddAction(StaticCastSharedPtr<FEdGraphSchemaAction>(Action.ToSharedPtr()));
+
+			return;
+		}
+	}
+
+	Super::GetGraphContextActions(ContextMenuBuilder);
+}
+
+const FPinConnectionResponse UCameraObjectGraphSchemaBase::CanCreateConnection(const UEdGraphPin* A, const UEdGraphPin* B) const
+{
+	// Check if we are connecting parameter pins of compatible types.
+	if ((A->PinType.PinCategory == PC_CameraParameter || 
+				A->PinType.PinCategory == PC_CameraVariableReference ||
+				A->PinType.PinCategory == PC_CameraContextData) && 
+			B->PinType.PinCategory == PC_Self &&
+			!A->bOrphanedPin)
+	{
+		UCameraObjectInterfaceParameterGraphNode* NodeB = Cast<UCameraObjectInterfaceParameterGraphNode>(B->GetOwningNode());
+		if (NodeB)
+		{
+			UCameraObjectInterfaceParameterGetter* ParameterGetter = NodeB->CastObject<UCameraObjectInterfaceParameterGetter>();
+			UCameraObjectInterfaceParameterBase* Parameter = ParameterGetter->GetInterfaceParameter();
+			UCameraObjectInterfaceBlendableParameter* BlendableParameter = Cast<UCameraObjectInterfaceBlendableParameter>(Parameter);
+			if (BlendableParameter && 
+					A->PinType.PinSubCategory == UEnum::GetValueAsName(BlendableParameter->ParameterType))
+			{
+				return FPinConnectionResponse(CONNECT_RESPONSE_MAKE, TEXT("Compatible pin types"));
+			}
+			UCameraObjectInterfaceDataParameter* DataParameter = Cast<UCameraObjectInterfaceDataParameter>(Parameter);
+			if (DataParameter && 
+					A->PinType.PinSubCategory == UEnum::GetValueAsName(DataParameter->DataType) &&
+					A->PinType.PinSubCategoryObject == DataParameter->DataTypeObject)
+			{
+				return FPinConnectionResponse(CONNECT_RESPONSE_MAKE, TEXT("Compatible pin types"));
+			}
+		}
+	}
+	else if (A->PinType.PinCategory == PC_Self && 
+			(B->PinType.PinCategory == PC_CameraParameter || 
+			 B->PinType.PinCategory == PC_CameraVariableReference ||
+			 B->PinType.PinCategory == PC_CameraContextData) &&
+			!B->bOrphanedPin)
+	{
+		UCameraObjectInterfaceParameterGraphNode* NodeA = Cast<UCameraObjectInterfaceParameterGraphNode>(A->GetOwningNode());
+		if (NodeA)
+		{
+			UCameraObjectInterfaceParameterGetter* ParameterGetter = NodeA->CastObject<UCameraObjectInterfaceParameterGetter>();
+			UCameraObjectInterfaceParameterBase* Parameter = ParameterGetter->GetInterfaceParameter();
+			UCameraObjectInterfaceBlendableParameter* BlendableParameter = Cast<UCameraObjectInterfaceBlendableParameter>(Parameter);
+			if (BlendableParameter && 
+					B->PinType.PinSubCategory == UEnum::GetValueAsName(BlendableParameter->ParameterType))
+			{
+				return FPinConnectionResponse(CONNECT_RESPONSE_MAKE, TEXT("Compatible pin types"));
+			}
+			UCameraObjectInterfaceDataParameter* DataParameter = Cast<UCameraObjectInterfaceDataParameter>(Parameter);
+			if (DataParameter && 
+					B->PinType.PinSubCategory == UEnum::GetValueAsName(DataParameter->DataType) &&
+					B->PinType.PinSubCategoryObject == DataParameter->DataTypeObject)
+			{
+				return FPinConnectionResponse(CONNECT_RESPONSE_MAKE, TEXT("Compatible pin types"));
+			}
+		}
+	}
+
+	return Super::CanCreateConnection(A, B);
+}
+
+bool UCameraObjectGraphSchemaBase::OnTryCreateCustomConnection(UEdGraphPin* A, UEdGraphPin* B) const
+{
+	UEdGraphPin* SourcePin = nullptr;
+	UEdGraphPin* TargetPin = nullptr;
+	if (A->Direction == EGPD_Output && B->Direction == EGPD_Input)
+	{
+		SourcePin = A;
+		TargetPin = B;
+	}
+	else if (B->Direction == EGPD_Output && A->Direction == EGPD_Input)
+	{
+		SourcePin = B;
+		TargetPin = A;
+	}
+	if (!ensure(SourcePin && TargetPin))
+	{
+		return false;
+	}
+
+	UObjectTreeGraphNode* SourceNode = Cast<UObjectTreeGraphNode>(SourcePin->GetOwningNode());
+	UObjectTreeGraphNode* TargetNode = Cast<UObjectTreeGraphNode>(TargetPin->GetOwningNode());
+	if (!SourceNode || !TargetNode)
+	{
+		return false;
+	}
+
+	ensure(SourceNode->GetGraph() == TargetNode->GetGraph());
+	UObjectTreeGraph* Graph = Cast<UObjectTreeGraph>(SourceNode->GetGraph());
+	if (!ensure(Graph))
+	{
+		return false;
+	}
+
+	UBaseCameraObject* CameraObject = Cast<UBaseCameraObject>(Graph->GetRootObject());
+	if (!CameraObject)
+	{
+		return false;
+	}
+
+	// See if we are connecting a parameter getter to a camera node's input property.
+	if (UCameraObjectInterfaceParameterGraphNode* ParameterNode = Cast<UCameraObjectInterfaceParameterGraphNode>(SourceNode))
+	{
+		if (SourcePin->PinType.PinCategory == PC_Self && 
+			(TargetPin->PinType.PinCategory == PC_CameraParameter || 
+			 TargetPin->PinType.PinCategory == PC_CameraVariableReference ||
+			 TargetPin->PinType.PinCategory == PC_CameraContextData))
+		{
+			CameraObject->Modify();
+
+			FCameraObjectConnection NewConnection;
+			NewConnection.Source = SourceNode->GetObject();
+			NewConnection.Target = TargetNode->GetObject();
+			NewConnection.TargetPropertyName = TargetPin->PinName;
+			CameraObject->Connections.Connections.Add(NewConnection);
+
+			return true;
+		}
+	}
+
+	// See if we are connecting one node's output to another node's input.
+	bool bDoFlowConnection = false;
+	if (SourcePin->PinType.PinCategory == PC_CameraVariableReference &&
+			(TargetPin->PinType.PinCategory == PC_CameraParameter ||
+			 TargetPin->PinType.PinCategory == PC_CameraVariableReference))
+	{
+		bDoFlowConnection = true;
+	}
+	else if (SourcePin->PinType.PinCategory == PC_CameraContextData &&
+			TargetPin->PinType.PinCategory == PC_CameraContextData)
+	{
+		bDoFlowConnection = true;
+	}
+	if (bDoFlowConnection)
+	{
+		CameraObject->Modify();
+
+		FCameraObjectConnection NewConnection;
+		NewConnection.Source = SourceNode->GetObject();
+		NewConnection.SourcePropertyName = SourcePin->PinName;
+		NewConnection.Target = TargetNode->GetObject();
+		NewConnection.TargetPropertyName = TargetPin->PinName;
+		CameraObject->Connections.Connections.Add(NewConnection);
+
+		return true;
+	}
+
+	return false;
+}
+
+bool UCameraObjectGraphSchemaBase::OnBreakCustomPinLinks(UEdGraphPin& TargetPin) const
+{
+	UObjectTreeGraphNode* Node = Cast<UObjectTreeGraphNode>(TargetPin.GetOwningNode());
+	if (!Node)
+	{
+		return false;
+	}
+
+	UObjectTreeGraph* Graph = Cast<UObjectTreeGraph>(Node->GetGraph());
+	if (!ensure(Graph))
+	{
+		return false;
+	}
+
+	UBaseCameraObject* CameraObject = Cast<UBaseCameraObject>(Graph->GetRootObject());
+	if (!CameraObject)
+	{
+		return false;
+	}
+
+	UObject* Object = Node->GetObject();
+	if (!Object)
+	{
+		return false;
+	}
+
+	bool bFoundAny = false;
+	for (auto It = CameraObject->Connections.Connections.CreateIterator(); It; ++It)
+	{
+		FCameraObjectConnection& Connection(*It);
+		if (
+				(Connection.Source == Object && TargetPin.PinName == Connection.SourcePropertyName) ||
+				(Connection.Target == Object && TargetPin.PinName == Connection.TargetPropertyName)
+		   )
+		{
+			if (!bFoundAny)
+			{
+				CameraObject->Modify();
+				bFoundAny = true;
+			}
+
+			It.RemoveCurrent();
+		}
+	}
+	return bFoundAny;
+}
+
+bool UCameraObjectGraphSchemaBase::OnBreakSingleCustomPinLink(UEdGraphPin* SourcePin, UEdGraphPin* TargetPin) const
+{
+	UObjectTreeGraphNode* SourceNode = Cast<UObjectTreeGraphNode>(SourcePin->GetOwningNode());
+	UObjectTreeGraphNode* TargetNode = Cast<UObjectTreeGraphNode>(TargetPin->GetOwningNode());
+	if (!SourceNode || !TargetNode)
+	{
+		return false;
+	}
+
+	UObjectTreeGraph* Graph = Cast<UObjectTreeGraph>(SourceNode->GetGraph());
+	if (!ensure(Graph))
+	{
+		return false;
+	}
+
+	UBaseCameraObject* CameraObject = Cast<UBaseCameraObject>(Graph->GetRootObject());
+	if (!CameraObject)
+	{
+		return false;
+	}
+
+	UObject* SourceObject = SourceNode->GetObject();
+	UObject* TargetObject = TargetNode->GetObject();
+
+	bool bFoundAny = false;
+	for (auto It = CameraObject->Connections.Connections.CreateIterator(); It; ++It)
+	{
+		FCameraObjectConnection& Connection(*It);
+		if (
+				(Connection.Source == SourceObject && SourcePin->PinName == Connection.SourcePropertyName) &&
+				(Connection.Target == TargetObject && TargetPin->PinName == Connection.TargetPropertyName)
+		   )
+		{
+			if (!bFoundAny)
+			{
+				CameraObject->Modify();
+				bFoundAny = true;
+			}
+
+			It.RemoveCurrent();
+		}
+	}
+	return bFoundAny;
+}
+
+FLinearColor UCameraObjectGraphSchemaBase::GetPinTypeColor(const FEdGraphPinType& PinType) const
+{
+	if (PinType.PinCategory == PC_CameraParameter || PinType.PinCategory == PC_CameraVariableReference)
+	{
+		const FName TypeName = PinType.PinSubCategory;
+		return PinColors.GetPinColor(TypeName);
+	}
+	if (PinType.PinCategory == PC_CameraContextData)
+	{
+		const FName DataTypeName = PinType.PinSubCategory;
+		return PinColors.GetContextDataPinColor(DataTypeName);
+	}
+
+	return UObjectTreeGraphSchema::GetPinTypeColor(PinType);
+}
+
+FCameraObjectGraphSchemaAction_NewInterfaceParameterNode::FCameraObjectGraphSchemaAction_NewInterfaceParameterNode()
+{
+}
+
+FCameraObjectGraphSchemaAction_NewInterfaceParameterNode::FCameraObjectGraphSchemaAction_NewInterfaceParameterNode(FText InNodeCategory, FText InMenuDesc, FText InToolTip, const int32 InGrouping, FText InKeywords)
+	: FEdGraphSchemaAction(InNodeCategory, InMenuDesc, InToolTip, InGrouping, InKeywords)
+{
+}
+
+UEdGraphNode* FCameraObjectGraphSchemaAction_NewInterfaceParameterNode::PerformAction(UEdGraph* ParentGraph, UEdGraphPin* FromPin, FPerformGraphActionLocation Location, bool bSelectNewNode)
+{
+	using namespace UE::Cameras;
+
+	UObjectTreeGraph* ObjectTreeGraph = Cast<UObjectTreeGraph>(ParentGraph);
+	if (!ensure(ObjectTreeGraph))
+	{
+		return nullptr;
+	}
+
+	UBaseCameraObject* CameraObject = Cast<UBaseCameraObject>(ObjectTreeGraph->GetRootObject());
+	if (!ensure(CameraObject))
+	{
+		return nullptr;
+	}
+
+	const FScopedTransaction Transaction(LOCTEXT("CreateNewNodeAction", "Create New Node"));
+
+	const UCameraObjectGraphSchemaBase* Schema = CastChecked<UCameraObjectGraphSchemaBase>(ParentGraph->GetSchema());
+
+	CameraObject->Modify();
+
+	// Create a new interface parameter and set it up based on the pin we're creating it from, if any.
+	UCameraObjectInterfaceParameterBase* NewInterfaceParameter = nullptr;
+	if (ParameterDefinition.ParameterType == ECameraObjectInterfaceParameterType::Blendable)
+	{
+		UCameraObjectInterfaceBlendableParameter* NewBlendableParameter = NewObject<UCameraObjectInterfaceBlendableParameter>(CameraObject, NAME_None, RF_Transactional);
+		NewBlendableParameter->ParameterType = ParameterDefinition.VariableType;
+		NewBlendableParameter->BlendableStructType = ParameterDefinition.BlendableStructType;
+		CameraObject->Interface.BlendableParameters.Add(NewBlendableParameter);
+		NewInterfaceParameter = NewBlendableParameter;
+	}
+	else if (ParameterDefinition.ParameterType == ECameraObjectInterfaceParameterType::Data)
+	{
+		UCameraObjectInterfaceDataParameter* NewDataParameter = NewObject<UCameraObjectInterfaceDataParameter>(CameraObject, NAME_None, RF_Transactional);
+		NewDataParameter->DataType = ParameterDefinition.DataType;
+		NewDataParameter->DataContainerType = ParameterDefinition.DataContainerType;
+		NewDataParameter->DataTypeObject = ParameterDefinition.DataTypeObject;
+		CameraObject->Interface.DataParameters.Add(NewDataParameter);
+		NewInterfaceParameter = NewDataParameter;
+	}
+
+	if (!ensure(NewInterfaceParameter))
+	{
+		return nullptr;
+	}
+
+	NewInterfaceParameter->InterfaceParameterName = FromPin ? FromPin->GetName() : NewInterfaceParameter->GetName();
+
+	// Set the value on the default parameters property bag to be the same as the value from which
+	// we pulled a connection.
+	if (FromPin)
+	{
+		if (UObjectTreeGraphNode* FromGraphNode = Cast<UObjectTreeGraphNode>(FromPin->GetOwningNode()))
+		{
+			if (UCameraNode* FromCameraNode = FromGraphNode->CastObject<UCameraNode>())
+			{
+				const FName TargetPropertyName = FromPin->GetFName();
+				FCameraObjectInterfaceParameterDefinition NewParameterDefinition;
+				NewInterfaceParameter->GetParameterDefinition(NewParameterDefinition);
+				FCameraObjectInterfaceParameterBuilder::SetDefaultParameterValue(
+						CameraObject, NewParameterDefinition, FromCameraNode, TargetPropertyName, true);
+			}
+		}
+	}
+
+	// Create a getter node for this parameter. It will then get connected inside AutowireNewNode.
+	UCameraObjectInterfaceParameterGetter* NewGetter = NewObject<UCameraObjectInterfaceParameterGetter>(CameraObject, NAME_None, RF_Transactional);
+	NewGetter->ParameterGuid = NewInterfaceParameter->GetGuid();
+	ensure(NewGetter->ParameterGuid.IsValid());
+
+	ObjectTreeGraph->Modify();
+
+	UObjectTreeGraphNode* NewGetterNode = CastChecked<UObjectTreeGraphNode>(Schema->CreateObjectNode(ObjectTreeGraph, NewGetter));
+
+	Schema->AddConnectableObject(ObjectTreeGraph, NewGetter);
+
+	NewGetterNode->NodePosX = Location.X;
+	NewGetterNode->NodePosY = Location.Y;
+	NewGetterNode->OnGraphNodeMoved(false);
+
+	NewGetterNode->AutowireNewNode(FromPin);
+
+	CameraObject->EventHandlers.Notify(&UE::Cameras::ICameraObjectEventHandler::OnCameraObjectInterfaceChanged);
+
+	return NewGetterNode;
+}
+
+FCameraObjectGraphSchemaAction_AddInterfaceParameterGetterNode::FCameraObjectGraphSchemaAction_AddInterfaceParameterGetterNode()
+{
+}
+
+FCameraObjectGraphSchemaAction_AddInterfaceParameterGetterNode::FCameraObjectGraphSchemaAction_AddInterfaceParameterGetterNode(FText InNodeCategory, FText InMenuDesc, FText InToolTip, const int32 InGrouping, FText InKeywords)
+	: FEdGraphSchemaAction(InNodeCategory, InMenuDesc, InToolTip, InGrouping, InKeywords)
+{
+}
+
+UEdGraphNode* FCameraObjectGraphSchemaAction_AddInterfaceParameterGetterNode::PerformAction(UEdGraph* ParentGraph, UEdGraphPin* FromPin, FPerformGraphActionLocation Location, bool bSelectNewNode)
+{
+	if (!ensure(InterfaceParameter))
+	{
+		return nullptr;
+	}
+
+	UObjectTreeGraph* ObjectTreeGraph = Cast<UObjectTreeGraph>(ParentGraph);
+	if (!ensure(ObjectTreeGraph))
+	{
+		return nullptr;
+	}
+
+	UBaseCameraObject* CameraObject = Cast<UBaseCameraObject>(ObjectTreeGraph->GetRootObject());
+	if (!ensure(CameraObject))
+	{
+		return nullptr;
+	}
+
+	const FScopedTransaction Transaction(LOCTEXT("CreateNewNodeAction", "Create New Node"));
+
+	const UCameraObjectGraphSchemaBase* Schema = CastChecked<UCameraObjectGraphSchemaBase>(ParentGraph->GetSchema());
+
+	CameraObject->Modify();
+
+	// Create a new getter for the given parameter.
+	UCameraObjectInterfaceParameterGetter* NewGetter = NewObject<UCameraObjectInterfaceParameterGetter>(CameraObject, NAME_None, RF_Transactional);
+	NewGetter->ParameterGuid = InterfaceParameter->GetGuid();
+	ensure(NewGetter->ParameterGuid.IsValid());
+
+	ParentGraph->Modify();
+	
+	UObjectTreeGraphNode* NewGetterNode = CastChecked<UObjectTreeGraphNode>(Schema->CreateObjectNode(ObjectTreeGraph, NewGetter));
+
+	Schema->AddConnectableObject(ObjectTreeGraph, NewGetter);
+
+	NewGetterNode->NodePosX = Location.X;
+	NewGetterNode->NodePosY = Location.Y;
+	NewGetterNode->OnGraphNodeMoved(false);
+
+	NewGetterNode->AutowireNewNode(FromPin);
+
+	CameraObject->EventHandlers.Notify(&UE::Cameras::ICameraObjectEventHandler::OnCameraObjectInterfaceChanged);
+
+	return NewGetterNode;
+}
+
+#undef LOCTEXT_NAMESPACE
+
diff --git a/Engine/Plugins/Cameras/GameplayCameras/Source/GameplayCamerasEditor/Private/Editors/CameraObjectInterfaceParameterGraphNode.cpp b/Engine/Plugins/Cameras/GameplayCameras/Source/GameplayCamerasEditor/Private/Editors/CameraObjectInterfaceParameterGraphNode.cpp
--- a/Engine/Plugins/Cameras/GameplayCameras/Source/GameplayCamerasEditor/Private/Editors/CameraObjectInterfaceParameterGraphNode.cpp
+++ b/Engine/Plugins/Cameras/GameplayCameras/Source/GameplayCamerasEditor/Private/Editors/CameraObjectInterfaceParameterGraphNode.cpp
@@ -15,7 +15,20 @@ UCameraObjectInterfaceParameterGraphNode::UCameraObjectInterfaceParameterGraphNo
 
 UCameraObjectInterfaceParameterBase* UCameraObjectInterfaceParameterGraphNode::GetInterfaceParameter() const
 {
-	return CastChecked<UCameraObjectInterfaceParameterBase>(GetObject(), ECastCheckedType::NullAllowed);
+	if (UCameraObjectInterfaceParameterGetter* GetterNode = CastChecked<UCameraObjectInterfaceParameterGetter>(GetObject(), ECastCheckedType::NullAllowed))
+	{
+		return GetterNode->GetInterfaceParameter();
+	}
+	return nullptr;
+}
+
+FText UCameraObjectInterfaceParameterGraphNode::GetNodeTitle(ENodeTitleType::Type TitleType) const
+{
+	if (UCameraObjectInterfaceParameterGetter* GetterNode = CastChecked<UCameraObjectInterfaceParameterGetter>(GetObject(), ECastCheckedType::NullAllowed))
+	{
+		return FText::FromString(GetterNode->GetInterfaceParameterName());
+	}
+	return FText();
 }
 
 TSharedPtr<SGraphNode> UCameraObjectInterfaceParameterGraphNode::CreateVisualWidget()
diff --git a/Engine/Plugins/Cameras/GameplayCameras/Source/GameplayCamerasEditor/Private/Editors/CameraRigCameraNodeGraphSchema.cpp b/Engine/Plugins/Cameras/GameplayCameras/Source/GameplayCamerasEditor/Private/Editors/CameraRigCameraNodeGraphSchema.cpp
--- a/Engine/Plugins/Cameras/GameplayCameras/Source/GameplayCamerasEditor/Private/Editors/CameraRigCameraNodeGraphSchema.cpp
+++ b/Engine/Plugins/Cameras/GameplayCameras/Source/GameplayCamerasEditor/Private/Editors/CameraRigCameraNodeGraphSchema.cpp
@@ -2,9 +2,11 @@
 
 #include "Editors/CameraRigCameraNodeGraphSchema.h"
 
+#include "Containers/Set.h"
 #include "Core/BlendCameraNode.h"
 #include "Core/CameraNodeHierarchy.h"
 #include "Core/CameraRigAsset.h"
+#include "Editors/CameraNodeGraphNode.h"
 #include "Editors/ObjectTreeGraph.h"
 #include "Editors/ObjectTreeGraphConfig.h"
 #include "Framework/Notifications/NotificationManager.h"
@@ -21,24 +23,29 @@ UCameraRigCameraNodeGraphSchema::UCameraRigCameraNodeGraphSchema(const FObjectIn
 {
 }
 
-FObjectTreeGraphConfig UCameraRigCameraNodeGraphSchema::BuildGraphConfig() const
+void UCameraRigCameraNodeGraphSchema::OnBuildGraphConfig(FObjectTreeGraphConfig& InOutGraphConfig) const
 {
-	const UGameplayCamerasEditorSettings* Settings = GetDefault<UGameplayCamerasEditorSettings>();
-
-	FObjectTreeGraphConfig GraphConfig;
+	Super::OnBuildGraphConfig(InOutGraphConfig);
 
-	BuildBaseGraphConfig(GraphConfig);
+	const UGameplayCamerasEditorSettings* Settings = GetDefault<UGameplayCamerasEditorSettings>();
 
-	GraphConfig.GraphName = UCameraRigAsset::NodeTreeGraphName;
-	GraphConfig.ConnectableObjectClasses.Add(UCameraNode::StaticClass());
-	GraphConfig.ConnectableObjectClasses.Add(UCameraRigAsset::StaticClass());
-	GraphConfig.NonConnectableObjectClasses.Add(UBlendCameraNode::StaticClass());
-	GraphConfig.ObjectClassConfigs.Emplace(UCameraRigAsset::StaticClass())
+	InOutGraphConfig.GraphName = UCameraRigAsset::NodeTreeGraphName;
+	InOutGraphConfig.GraphDisplayInfo.PlainName = LOCTEXT("NodeGraphPlainName", "CameraNodes");
+	InOutGraphConfig.GraphDisplayInfo.DisplayName = LOCTEXT("NodeGraphDisplayName", "Camera Nodes");
+	InOutGraphConfig.ConnectableObjectClasses.Add(UCameraNode::StaticClass());
+	InOutGraphConfig.ConnectableObjectClasses.Add(UCameraRigAsset::StaticClass());
+	InOutGraphConfig.NonConnectableObjectClasses.Add(UBlendCameraNode::StaticClass());
+	InOutGraphConfig.ObjectClassConfigs.Emplace(UCameraNode::StaticClass())
+		.StripDisplayNameSuffix(TEXT("Camera Node"))
+		.CreateCategoryMetaData(TEXT("CameraNodeCategories"))
+		.NodeTitleColor(Settings->CameraNodeTitleColor)
+		.GraphNodeClass(UCameraNodeGraphNode::StaticClass());
+	InOutGraphConfig.ObjectClassConfigs.Emplace(UCameraRigAsset::StaticClass())
 		.OnlyAsRoot()
 		.HasSelfPin(false)
 		.NodeTitleUsesObjectName(true)
 		.NodeTitleColor(Settings->CameraRigAssetTitleColor);
-	GraphConfig.ObjectClassConfigs.Emplace(UArrayCameraNode::StaticClass())
+	InOutGraphConfig.ObjectClassConfigs.Emplace(UArrayCameraNode::StaticClass())
 		.OnSetupNewObject(FOnSetupNewObject::CreateLambda([](UObject* NewObject)
 				{
 					// Add two new pins by default.
@@ -46,15 +53,14 @@ FObjectTreeGraphConfig UCameraRigCameraNodeGraphSchema::BuildGraphConfig() const
 					ArrayNode->Children.AddDefaulted();
 					ArrayNode->Children.AddDefaulted();
 				}));
-
-	return GraphConfig;
 }
 
 void UCameraRigCameraNodeGraphSchema::CollectAllObjects(UObjectTreeGraph* InGraph, TSet<UObject*>& OutAllObjects) const
 {
 	using namespace UE::Cameras;
 
-	Super::CollectAllObjects(InGraph, OutAllObjects);
+	// Start with the graph objects from the root interface.
+	CollectAllConnectableObjectsFromRootInterface(InGraph, OutAllObjects, false);
 
 	// See if we are missing objects from AllNodeTreeObjects... if so, add them and notify the user.
 	UCameraRigAsset* CameraRig = Cast<UCameraRigAsset>(InGraph->GetRootObject());
diff --git a/Engine/Plugins/Cameras/GameplayCameras/Source/GameplayCamerasEditor/Private/Editors/CameraRigTransitionGraphSchema.cpp b/Engine/Plugins/Cameras/GameplayCameras/Source/GameplayCamerasEditor/Private/Editors/CameraRigTransitionGraphSchema.cpp
--- a/Engine/Plugins/Cameras/GameplayCameras/Source/GameplayCamerasEditor/Private/Editors/CameraRigTransitionGraphSchema.cpp
+++ b/Engine/Plugins/Cameras/GameplayCameras/Source/GameplayCamerasEditor/Private/Editors/CameraRigTransitionGraphSchema.cpp
@@ -13,6 +13,8 @@
 
 void UCameraRigTransitionGraphSchema::OnBuildGraphConfig(FObjectTreeGraphConfig& InOutGraphConfig) const
 {
+	Super::OnBuildGraphConfig(InOutGraphConfig);
+
 	const UGameplayCamerasEditorSettings* Settings = GetDefault<UGameplayCamerasEditorSettings>();
 
 	InOutGraphConfig.GraphName = UCameraRigAsset::TransitionsGraphName;
diff --git a/Engine/Plugins/Cameras/GameplayCameras/Source/GameplayCamerasEditor/Private/Editors/CameraRigTransitionGraphSchemaBase.cpp b/Engine/Plugins/Cameras/GameplayCameras/Source/GameplayCamerasEditor/Private/Editors/CameraRigTransitionGraphSchemaBase.cpp
--- a/Engine/Plugins/Cameras/GameplayCameras/Source/GameplayCamerasEditor/Private/Editors/CameraRigTransitionGraphSchemaBase.cpp
+++ b/Engine/Plugins/Cameras/GameplayCameras/Source/GameplayCamerasEditor/Private/Editors/CameraRigTransitionGraphSchemaBase.cpp
@@ -6,6 +6,7 @@
 #include "Core/CameraRigAsset.h"
 #include "Core/CameraRigTransition.h"
 #include "EdGraph/EdGraphPin.h"
+#include "Editors/CameraNodeGraphNode.h"
 #include "Editors/ObjectTreeGraph.h"
 #include "Editors/ObjectTreeGraphConfig.h"
 #include "Editors/ObjectTreeGraphNode.h"
@@ -73,37 +74,34 @@ private:
 
 }  // namespace UE::Cameras
 
-FObjectTreeGraphConfig UCameraRigTransitionGraphSchemaBase::BuildGraphConfig() const
+void UCameraRigTransitionGraphSchemaBase::OnBuildGraphConfig(FObjectTreeGraphConfig& InOutGraphConfig) const
 {
 	using namespace UE::Cameras;
 
+	Super::OnBuildGraphConfig(InOutGraphConfig);
+
 	const UGameplayCamerasEditorSettings* Settings = GetDefault<UGameplayCamerasEditorSettings>();
 
-	FObjectTreeGraphConfig GraphConfig;
-	GraphConfig.ConnectableObjectClasses.Add(UCameraRigTransition::StaticClass());
-	GraphConfig.ConnectableObjectClasses.Add(UCameraRigTransitionCondition::StaticClass());
-	GraphConfig.ConnectableObjectClasses.Add(UBlendCameraNode::StaticClass());
-	GraphConfig.DefaultSelfPinName = NAME_None;
-	GraphConfig.ObjectClassConfigs.Emplace(UCameraRigTransition::StaticClass())
+	InOutGraphConfig.ConnectableObjectClasses.Add(UCameraRigTransition::StaticClass());
+	InOutGraphConfig.ConnectableObjectClasses.Add(UCameraRigTransitionCondition::StaticClass());
+	InOutGraphConfig.ConnectableObjectClasses.Add(UBlendCameraNode::StaticClass());
+	InOutGraphConfig.ObjectClassConfigs.Emplace(UCameraRigTransition::StaticClass())
 		.NodeTitleColor(Settings->CameraRigTransitionTitleColor);
-	GraphConfig.ObjectClassConfigs.Emplace(UCameraRigTransitionCondition::StaticClass())
+	InOutGraphConfig.ObjectClassConfigs.Emplace(UCameraRigTransitionCondition::StaticClass())
 		.StripDisplayNameSuffix(TEXT("Transition Condition"))
 		.NodeTitleColor(Settings->CameraRigTransitionConditionTitleColor);
-	GraphConfig.ObjectClassConfigs.Emplace(UBlendCameraNode::StaticClass())
+	InOutGraphConfig.ObjectClassConfigs.Emplace(UBlendCameraNode::StaticClass())
 		.StripDisplayNameSuffix(TEXT("Camera Node"))
 		.CreateCategoryMetaData(TEXT("CameraNodeCategories"))
-		.NodeTitleColor(Settings->CameraBlendNodeTitleColor);
-
-	OnBuildGraphConfig(GraphConfig);
-
-	return GraphConfig;
+		.NodeTitleColor(Settings->CameraBlendNodeTitleColor)
+		.GraphNodeClass(UCameraNodeGraphNode::StaticClass());
 }
 
 void UCameraRigTransitionGraphSchemaBase::CollectAllObjects(UObjectTreeGraph* InGraph, TSet<UObject*>& OutAllObjects) const
 {
 	using namespace UE::Cameras;
 
-	// Only get the graph objects from the root interface.
+	// Start with the graph objects from the root interface.
 	CollectAllConnectableObjectsFromRootInterface(InGraph, OutAllObjects, false);
 
 	// See if we are missing objects from AllTransitionsObjects... if so, add them and notify the user.
diff --git a/Engine/Plugins/Cameras/GameplayCameras/Source/GameplayCamerasEditor/Private/Editors/CameraShakeCameraNodeGraphSchema.cpp b/Engine/Plugins/Cameras/GameplayCameras/Source/GameplayCamerasEditor/Private/Editors/CameraShakeCameraNodeGraphSchema.cpp
--- a/Engine/Plugins/Cameras/GameplayCameras/Source/GameplayCamerasEditor/Private/Editors/CameraShakeCameraNodeGraphSchema.cpp
+++ b/Engine/Plugins/Cameras/GameplayCameras/Source/GameplayCamerasEditor/Private/Editors/CameraShakeCameraNodeGraphSchema.cpp
@@ -17,24 +17,30 @@ UCameraShakeCameraNodeGraphSchema::UCameraShakeCameraNodeGraphSchema(const FObje
 {
 }
 
-FObjectTreeGraphConfig UCameraShakeCameraNodeGraphSchema::BuildGraphConfig() const
+void UCameraShakeCameraNodeGraphSchema::CollectAllObjects(UObjectTreeGraph* InGraph, TSet<UObject*>& OutAllObjects) const
 {
-	const UGameplayCamerasEditorSettings* Settings = GetDefault<UGameplayCamerasEditorSettings>();
+	using namespace UE::Cameras;
+
+	// Only get the graph objects from the root interface.
+	CollectAllConnectableObjectsFromRootInterface(InGraph, OutAllObjects, false);
+}
 
-	FObjectTreeGraphConfig GraphConfig;
+void UCameraShakeCameraNodeGraphSchema::OnBuildGraphConfig(FObjectTreeGraphConfig& InOutGraphConfig) const
+{
+	Super::OnBuildGraphConfig(InOutGraphConfig);
 
-	BuildBaseGraphConfig(GraphConfig);
+	const UGameplayCamerasEditorSettings* Settings = GetDefault<UGameplayCamerasEditorSettings>();
 
-	GraphConfig.ConnectableObjectClasses.Add(UShakeCameraNode::StaticClass());
-	GraphConfig.ConnectableObjectClasses.Add(USimpleFixedTimeBlendCameraNode::StaticClass());
-	GraphConfig.ConnectableObjectClasses.Add(UCameraShakeAsset::StaticClass());
-	GraphConfig.ObjectClassConfigs.Emplace(UCameraShakeAsset::StaticClass())
+	InOutGraphConfig.GraphDisplayInfo.PlainName = LOCTEXT("NodeGraphPlainName", "CameraShakeNodes");
+	InOutGraphConfig.GraphDisplayInfo.DisplayName = LOCTEXT("NodeGraphDisplayName", "Camera Shake Nodes");
+	InOutGraphConfig.ConnectableObjectClasses.Add(UShakeCameraNode::StaticClass());
+	InOutGraphConfig.ConnectableObjectClasses.Add(USimpleFixedTimeBlendCameraNode::StaticClass());
+	InOutGraphConfig.ConnectableObjectClasses.Add(UCameraShakeAsset::StaticClass());
+	InOutGraphConfig.ObjectClassConfigs.Emplace(UCameraShakeAsset::StaticClass())
 		.OnlyAsRoot()
 		.HasSelfPin(false)
 		.NodeTitleUsesObjectName(true)
 		.NodeTitleColor(Settings->CameraShakeAssetTitleColor);
-
-	return GraphConfig;
 }
 
 #undef LOCTEXT_NAMESPACE
diff --git a/Engine/Plugins/Cameras/GameplayCameras/Source/GameplayCamerasEditor/Private/Editors/CameraSharedTransitionGraphSchema.cpp b/Engine/Plugins/Cameras/GameplayCameras/Source/GameplayCamerasEditor/Private/Editors/CameraSharedTransitionGraphSchema.cpp
--- a/Engine/Plugins/Cameras/GameplayCameras/Source/GameplayCamerasEditor/Private/Editors/CameraSharedTransitionGraphSchema.cpp
+++ b/Engine/Plugins/Cameras/GameplayCameras/Source/GameplayCamerasEditor/Private/Editors/CameraSharedTransitionGraphSchema.cpp
@@ -12,6 +12,8 @@
 
 void UCameraSharedTransitionGraphSchema::OnBuildGraphConfig(FObjectTreeGraphConfig& InOutGraphConfig) const
 {
+	Super::OnBuildGraphConfig(InOutGraphConfig);
+
 	const UGameplayCamerasEditorSettings* Settings = GetDefault<UGameplayCamerasEditorSettings>();
 
 	InOutGraphConfig.GraphName = UCameraAsset::SharedTransitionsGraphName;
diff --git a/Engine/Plugins/Cameras/GameplayCameras/Source/GameplayCamerasEditor/Private/Editors/ObjectTreeGraphSchema.cpp b/Engine/Plugins/Cameras/GameplayCameras/Source/GameplayCamerasEditor/Private/Editors/ObjectTreeGraphSchema.cpp
--- a/Engine/Plugins/Cameras/GameplayCameras/Source/GameplayCamerasEditor/Private/Editors/ObjectTreeGraphSchema.cpp
+++ b/Engine/Plugins/Cameras/GameplayCameras/Source/GameplayCamerasEditor/Private/Editors/ObjectTreeGraphSchema.cpp
@@ -876,7 +876,10 @@ void UObjectTreeGraphSchema::ApplyConnection(UEdGraphPin* A, UEdGraphPin* B) con
 		ValuePin = B;
 	}
 
-	checkf(PropertyPin && ValuePin, TEXT("Invalid pins passed for setting property values."));
+	if (!ensureMsgf(PropertyPin && ValuePin, TEXT("Invalid pins passed for setting property values.")))
+	{
+		return;
+	}
 
 	UObjectTreeGraphNode* PropertyNode = CastChecked<UObjectTreeGraphNode>(PropertyPin->GetOwningNode());
 	UObjectTreeGraphNode* ValueNode = CastChecked<UObjectTreeGraphNode>(ValuePin->GetOwningNode());
@@ -960,7 +963,10 @@ void UObjectTreeGraphSchema::ApplyDisconnection(UEdGraphPin* TargetPin) const
 	{
 		TargetPin = TargetPin->LinkedTo[0];
 	}
-	check(TargetPin->PinType.PinCategory == PC_Property);
+	if (!ensure(TargetPin->PinType.PinCategory == PC_Property))
+	{
+		return;
+	}
 
 	UObjectTreeGraphNode* PropertyNode = Cast<UObjectTreeGraphNode>(TargetPin->GetOwningNode());
 	check(PropertyNode);
diff --git a/Engine/Plugins/Cameras/GameplayCameras/Source/GameplayCamerasEditor/Private/Editors/SCameraNodeGraphNode.cpp b/Engine/Plugins/Cameras/GameplayCameras/Source/GameplayCamerasEditor/Private/Editors/SCameraNodeGraphNode.cpp
--- a/Engine/Plugins/Cameras/GameplayCameras/Source/GameplayCamerasEditor/Private/Editors/SCameraNodeGraphNode.cpp
+++ b/Engine/Plugins/Cameras/GameplayCameras/Source/GameplayCamerasEditor/Private/Editors/SCameraNodeGraphNode.cpp
@@ -2,7 +2,7 @@
 
 #include "Editors/SCameraNodeGraphNode.h"
 
-#include "Editors/CameraNodeGraphSchema.h"
+#include "Editors/CameraObjectGraphSchemaBase.h"
 #include "Styles/GameplayCamerasEditorStyle.h"
 
 void SCameraNodeGraphNode::Construct(const FArguments& InArgs)
@@ -18,12 +18,12 @@ TSharedPtr<SGraphPin> SCameraNodeGraphNode::CreatePinWidget(UEdGraphPin* InPin)
 
 	TSharedPtr<SGraphPin> PinWidget = SObjectTreeGraphNode::CreatePinWidget(InPin);
 
-	if (const UCameraNodeGraphSchema* GraphSchema = Cast<const UCameraNodeGraphSchema>(InPin->GetSchema()))
+	if (const UCameraObjectGraphSchemaBase* GraphSchema = Cast<const UCameraObjectGraphSchemaBase>(InPin->GetSchema()))
 	{
 		TSharedRef<FGameplayCamerasEditorStyle> GraphStyle = FGameplayCamerasEditorStyle::Get();
 
-		if (InPin->PinType.PinCategory == UCameraNodeGraphSchema::PC_CameraParameter ||
-				InPin->PinType.PinCategory == UCameraNodeGraphSchema::PC_CameraVariableReference)
+		if (InPin->PinType.PinCategory == UCameraObjectGraphSchemaBase::PC_CameraParameter ||
+				InPin->PinType.PinCategory == UCameraObjectGraphSchemaBase::PC_CameraVariableReference)
 		{
 			const FSlateBrush* ConnectedBrush = GraphStyle->GetBrush("Graph.CameraRigParameterPin.Connected");
 			const FSlateBrush* DisconnectedBrush = GraphStyle->GetBrush("Graph.CameraRigParameterPin.Disconnected");
diff --git a/Engine/Plugins/Cameras/GameplayCameras/Source/GameplayCamerasEditor/Private/Editors/SCameraObjectInterfaceParameterGraphNode.cpp b/Engine/Plugins/Cameras/GameplayCameras/Source/GameplayCamerasEditor/Private/Editors/SCameraObjectInterfaceParameterGraphNode.cpp
--- a/Engine/Plugins/Cameras/GameplayCameras/Source/GameplayCamerasEditor/Private/Editors/SCameraObjectInterfaceParameterGraphNode.cpp
+++ b/Engine/Plugins/Cameras/GameplayCameras/Source/GameplayCamerasEditor/Private/Editors/SCameraObjectInterfaceParameterGraphNode.cpp
@@ -132,9 +132,9 @@ TSharedPtr<SGraphPin> SCameraObjectInterfaceParameterGraphNode::CreatePinWidget(
 
 FText SCameraObjectInterfaceParameterGraphNode::GetInterfaceParameterName() const
 {
-	if (UCameraObjectInterfaceParameterBase* InterfaceParameter = GetObjectGraphNode()->CastObject<UCameraObjectInterfaceParameterBase>())
+	if (UCameraObjectInterfaceParameterGetter* GetterNode = GetObjectGraphNode()->CastObject<UCameraObjectInterfaceParameterGetter>())
 	{
-		return FText::FromString(InterfaceParameter->InterfaceParameterName);
+		return FText::FromString(GetterNode->GetInterfaceParameterName());
 	}
 	return LOCTEXT("InvalidParameterName", "Invalid");
 }
diff --git a/Engine/Plugins/Cameras/GameplayCameras/Source/GameplayCamerasEditor/Private/Editors/SCameraObjectInterfaceParametersPanel.cpp b/Engine/Plugins/Cameras/GameplayCameras/Source/GameplayCamerasEditor/Private/Editors/SCameraObjectInterfaceParametersPanel.cpp
new file mode 100644
--- /dev/null
+++ b/Engine/Plugins/Cameras/GameplayCameras/Source/GameplayCamerasEditor/Private/Editors/SCameraObjectInterfaceParametersPanel.cpp
@@ -0,0 +1,856 @@
+// Copyright Epic Games, Inc. All Rights Reserved.
+
+#include "Editors/SCameraObjectInterfaceParametersPanel.h"
+
+#include "Editors/CameraNodeGraphDragDropOp.h"
+#include "Framework/MultiBox/MultiBoxBuilder.h"
+#include "IGameplayCamerasModule.h"
+#include "SPinTypeSelector.h"
+#include "ScopedTransaction.h"
+#include "Styles/GameplayCamerasEditorStyle.h"
+#include "Toolkits/CameraObjectInterfaceParametersToolkit.h"
+#include "Types/SlateEnums.h"
+#include "Widgets/Input/SButton.h"
+#include "Widgets/Input/SCheckBox.h"
+#include "Widgets/SCompoundWidget.h"
+#include "Widgets/Text/SInlineEditableTextBlock.h"
+#include "Widgets/Text/SRichTextBlock.h"
+#include "Widgets/Text/STextBlock.h"
+#include "Widgets/Views/SListView.h"
+
+#include UE_INLINE_GENERATED_CPP_BY_NAME(SCameraObjectInterfaceParametersPanel)
+
+#define LOCTEXT_NAMESPACE "SCameraObjectInterfaceParametersPanel"
+
+namespace UE::Cameras
+{
+
+static const FName ParameterTypeColumn(TEXT("ParameterType"));
+static const FName ParameterNameColumn(TEXT("ParameterName"));
+
+/**
+ * List entry for any interface parameter panel.
+ */
+template<typename ParameterType>
+class SCameraObjectInterfaceParameterTableRowBase : public SMultiColumnTableRow<TObjectPtr<ParameterType>>
+{
+public:
+
+	SLATE_BEGIN_ARGS(SCameraObjectInterfaceParameterTableRowBase<ParameterType>)
+	{}
+		SLATE_ARGUMENT(TObjectPtr<ParameterType>, Item)
+	SLATE_END_ARGS()
+
+	typedef SMultiColumnTableRow<TObjectPtr<ParameterType>> FSuperRowType;
+	typedef typename STableRow<TObjectPtr<ParameterType>>::FArguments FTableRowArgs;
+		
+	void Construct(const FArguments& Args, const TSharedRef<STableViewBase>& OwnerTable)
+	{
+		Item = Args._Item;
+
+		FSuperRowType::Construct(
+				FTableRowArgs(),
+				OwnerTable);
+	}
+
+	void EnterNameEditingMode()
+	{
+		NameTextBlock->EnterEditingMode();
+	}
+
+protected:
+
+	virtual TSharedRef<SWidget> GenerateWidgetForColumn(const FName& InColumnName) override
+	{
+		if (InColumnName == ParameterNameColumn)
+		{
+			return SAssignNew(NameTextBlock, SInlineEditableTextBlock)
+				.IsSelected(this, &SCameraObjectInterfaceParameterTableRowBase<ParameterType>::IsSelected)
+				.Text_Lambda([this]() { return FText::FromString(Item->InterfaceParameterName); })
+				.OnTextCommitted(this, &SCameraObjectInterfaceParameterTableRowBase<ParameterType>::OnParameterNameTextCommitted)
+				.ToolTipText(this, &SCameraObjectInterfaceParameterTableRowBase<ParameterType>::GetParameterToolTip);
+		}
+
+		return SNullWidget::NullWidget;
+	}
+
+	virtual FReply OnDragDetected(const FGeometry& MyGeometry, const FPointerEvent& MouseEvent) override
+	{
+		if (MouseEvent.IsMouseButtonDown(EKeys::LeftMouseButton))
+		{
+			return FReply::Handled().BeginDragDrop(FCameraNodeGraphInterfaceParameterDragDropOp::New(Item));
+		}
+		return SMultiColumnTableRow<TObjectPtr<ParameterType>>::OnDragDetected(MyGeometry, MouseEvent);
+	}
+
+	void OnParameterNameTextCommitted(const FText& Text, ETextCommit::Type CommitType)
+	{
+		const FScopedTransaction Transaction(LOCTEXT("RenameInterfaceParameterTransaction", "Rename Interface Parameter"));
+
+		Item->Modify();
+		Item->InterfaceParameterName = Text.ToString();
+	}
+
+	FText GetParameterToolTip() const
+	{
+		const FGuid Guid = Item->GetGuid();
+		return FText::FromString(Guid.ToString(EGuidFormats::DigitsWithHyphensInParentheses));
+	}
+
+protected:
+
+	TObjectPtr<ParameterType> Item;
+	TSharedPtr<SInlineEditableTextBlock> NameTextBlock;
+};
+
+/**
+ * List entry for the blendable parameters panel.
+ */
+class SCameraObjectInterfaceBlendableParameterTableRow : public SCameraObjectInterfaceParameterTableRowBase<UCameraObjectInterfaceBlendableParameter>
+{
+public:
+
+	void Construct(const FArguments& Args, const TSharedRef<STableViewBase>& OwnerTable)
+	{
+		SCameraObjectInterfaceParameterTableRowBase<UCameraObjectInterfaceBlendableParameter>::Construct(Args, OwnerTable);
+	}
+
+protected:
+
+	virtual TSharedRef<SWidget> GenerateWidgetForColumn(const FName& InColumnName) override
+	{
+		if (InColumnName == ParameterTypeColumn)
+		{
+			TSharedRef<FGameplayCamerasEditorStyle> GameplayCamerasEditorStyle = FGameplayCamerasEditorStyle::Get();
+
+			const UEdGraphSchema_K2* K2Schema = GetDefault<UEdGraphSchema_CameraNodeK2>();
+
+			return SNew(SBox)
+				.MinDesiredWidth(125.f)
+				.Padding(this, &SCameraObjectInterfaceBlendableParameterTableRow::GetPinTypeSelectorPadding)
+				[
+					SNew(SPinTypeSelector, FGetPinTypeTree::CreateSP(this, &SCameraObjectInterfaceBlendableParameterTableRow::GetPinTypeTreeInfos))
+					.OnPinTypeChanged(this, &SCameraObjectInterfaceBlendableParameterTableRow::OnBlendableParameterPinTypeChanged)
+					.TargetPinType(this, &SCameraObjectInterfaceBlendableParameterTableRow::GetBlendableParameterPinType)
+					.ReadOnly(this, &SCameraObjectInterfaceBlendableParameterTableRow::IsPinTypeSelectorReadOnly)
+					.Schema(K2Schema)
+					.bAllowArrays(false)
+				];
+		}
+
+		return SCameraObjectInterfaceParameterTableRowBase<UCameraObjectInterfaceBlendableParameter>::GenerateWidgetForColumn(InColumnName);
+	}
+
+	FMargin GetPinTypeSelectorPadding() const
+	{
+		// Add some horizontal margin when the pin type selector is read-only, so that the transition on hover is seamless.
+		return IsHovered() ? FMargin(0) : FMargin(7, 0);
+	}
+
+	void GetPinTypeTreeInfos(TArray<FPinTypeTreeItem>& TypeTree, ETypeTreeFilter TypeTreeFilter) const
+	{
+		using FPinTypeTreeInfo = UEdGraphSchema_K2::FPinTypeTreeInfo;
+
+		const UEdGraphSchema_K2* K2Schema = GetDefault<UEdGraphSchema_CameraNodeK2>();
+
+		TypeTree.Reset();
+
+		TypeTree.Add(MakeShared<FPinTypeTreeInfo>(UEdGraphSchema_K2::GetCategoryText(UEdGraphSchema_K2::PC_Boolean, true), UEdGraphSchema_K2::PC_Boolean, K2Schema, LOCTEXT("BooleanType", "True or false value")));
+		TypeTree.Add(MakeShared<FPinTypeTreeInfo>(UEdGraphSchema_K2::GetCategoryText(UEdGraphSchema_K2::PC_Int, true), UEdGraphSchema_K2::PC_Int, K2Schema, LOCTEXT("IntegerType", "Integer number")));
+		TypeTree.Add(MakeShared<FPinTypeTreeInfo>(UEdGraphSchema_K2::GetCategoryText(UEdGraphSchema_K2::PC_Float, true), UEdGraphSchema_K2::PC_Float, K2Schema, LOCTEXT("FloatType", "Floating point number")));
+		TypeTree.Add(MakeShared<FPinTypeTreeInfo>(UEdGraphSchema_K2::GetCategoryText(UEdGraphSchema_K2::PC_Double, true), UEdGraphSchema_K2::PC_Double, K2Schema, LOCTEXT("DoubleType", "Double precision floating point number")));
+
+		TypeTree.Add(MakeShared<FPinTypeTreeInfo>(UEdGraphSchema_K2::PC_Struct, TVariantStructure<FVector2f>::Get(), LOCTEXT("Vector2fType", "A 2D vector")));
+		TypeTree.Add(MakeShared<FPinTypeTreeInfo>(UEdGraphSchema_K2::PC_Struct, TBaseStructure<FVector2D>::Get(), LOCTEXT("Vector2dType", "A double precision 2D vector")));
+		TypeTree.Add(MakeShared<FPinTypeTreeInfo>(UEdGraphSchema_K2::PC_Struct, TVariantStructure<FVector3f>::Get(), LOCTEXT("Vector3fType", "A 3D vector")));
+		TypeTree.Add(MakeShared<FPinTypeTreeInfo>(UEdGraphSchema_K2::PC_Struct, TBaseStructure<FVector>::Get(), LOCTEXT("Vector3dType", "A double precision 3D vector")));
+		TypeTree.Add(MakeShared<FPinTypeTreeInfo>(UEdGraphSchema_K2::PC_Struct, TVariantStructure<FVector4f>::Get(), LOCTEXT("Vector4fType", "A 4D vector")));
+		TypeTree.Add(MakeShared<FPinTypeTreeInfo>(UEdGraphSchema_K2::PC_Struct, TBaseStructure<FVector4>::Get(), LOCTEXT("Vector4dType", "A double precision 4D vector")));
+		TypeTree.Add(MakeShared<FPinTypeTreeInfo>(UEdGraphSchema_K2::PC_Struct, TVariantStructure<FRotator3f>::Get(), LOCTEXT("Rotator3fType", "A 3D rotation")));
+		TypeTree.Add(MakeShared<FPinTypeTreeInfo>(UEdGraphSchema_K2::PC_Struct, TBaseStructure<FRotator>::Get(), LOCTEXT("Rotator3dType", "A double precision 3D rotation")));
+		TypeTree.Add(MakeShared<FPinTypeTreeInfo>(UEdGraphSchema_K2::PC_Struct, TVariantStructure<FTransform3f>::Get(), LOCTEXT("Transform3fType", "A 3D transformation")));
+		TypeTree.Add(MakeShared<FPinTypeTreeInfo>(UEdGraphSchema_K2::PC_Struct, TBaseStructure<FTransform>::Get(), LOCTEXT("Transform3dType", "A double precision 3D transformation")));
+
+		TSharedPtr<FPinTypeTreeInfo> Structs = MakeShared<FPinTypeTreeInfo>(
+				LOCTEXT("BlendableStructPinTypeLabel", "Blendable Structures"),
+				UEdGraphSchema_K2::PC_Struct,
+				K2Schema,
+				LOCTEXT("BlendableStructPinTypeToolTip", "Blendable structure types"),
+				true);
+		IGameplayCamerasModule& GameplayCamerasModule = IGameplayCamerasModule::Get();
+		for (const FBlendableStructInfo& BlendableStruct : GameplayCamerasModule.GetBlendableStructs())
+		{
+			if (const UScriptStruct* StructType = BlendableStruct.StructType)
+			{
+				Structs->Children.Add(
+						MakeShared<UEdGraphSchema_K2::FPinTypeTreeInfo>(
+							UEdGraphSchema_K2::PC_Struct,
+							const_cast<UScriptStruct*>(StructType),
+							StructType->GetToolTipText(),
+							false,
+							(uint8)EObjectReferenceType::NotAnObject));
+			}
+		}
+		TypeTree.Add(Structs);
+	}
+
+	FEdGraphPinType GetBlendableParameterPinType() const
+	{
+		FEdGraphPinType PinType;
+		switch (Item->ParameterType)
+		{
+			case ECameraVariableType::Boolean:
+				PinType.PinCategory = UEdGraphSchema_K2::PC_Boolean;
+				break;
+			case ECameraVariableType::Integer32:
+				PinType.PinCategory = UEdGraphSchema_K2::PC_Int;
+				break;
+			case ECameraVariableType::Float:
+				PinType.PinCategory = UEdGraphSchema_K2::PC_Float;
+				break;
+			case ECameraVariableType::Double:
+				PinType.PinCategory = UEdGraphSchema_K2::PC_Double;
+				break;
+			case ECameraVariableType::Vector2f:
+				PinType.PinCategory = UEdGraphSchema_K2::PC_Struct;
+				PinType.PinSubCategoryObject = TVariantStructure<FVector2f>::Get();
+				break;
+			case ECameraVariableType::Vector2d:
+				PinType.PinCategory = UEdGraphSchema_K2::PC_Struct;
+				PinType.PinSubCategoryObject = TBaseStructure<FVector2D>::Get();
+				break;
+			case ECameraVariableType::Vector3f:
+				PinType.PinCategory = UEdGraphSchema_K2::PC_Struct;
+				PinType.PinSubCategoryObject = TVariantStructure<FVector3f>::Get();
+				break;
+			case ECameraVariableType::Vector3d:
+				PinType.PinCategory = UEdGraphSchema_K2::PC_Struct;
+				PinType.PinSubCategoryObject = TBaseStructure<FVector>::Get();
+				break;
+			case ECameraVariableType::Vector4f:
+				PinType.PinCategory = UEdGraphSchema_K2::PC_Struct;
+				PinType.PinSubCategoryObject = TVariantStructure<FVector4f>::Get();
+				break;
+			case ECameraVariableType::Vector4d:
+				PinType.PinCategory = UEdGraphSchema_K2::PC_Struct;
+				PinType.PinSubCategoryObject = TBaseStructure<FVector4>::Get();
+				break;
+			case ECameraVariableType::Rotator3f:
+				PinType.PinCategory = UEdGraphSchema_K2::PC_Struct;
+				PinType.PinSubCategoryObject = TVariantStructure<FRotator3f>::Get();
+				break;
+			case ECameraVariableType::Rotator3d:
+				PinType.PinCategory = UEdGraphSchema_K2::PC_Struct;
+				PinType.PinSubCategoryObject = TBaseStructure<FRotator>::Get();
+				break;
+			case ECameraVariableType::Transform3f:
+				PinType.PinCategory = UEdGraphSchema_K2::PC_Struct;
+				PinType.PinSubCategoryObject = TVariantStructure<FTransform3f>::Get();
+				break;
+			case ECameraVariableType::Transform3d:
+				PinType.PinCategory = UEdGraphSchema_K2::PC_Struct;
+				PinType.PinSubCategoryObject = TBaseStructure<FTransform>::Get();
+				break;
+			case ECameraVariableType::BlendableStruct:
+				PinType.PinCategory = UEdGraphSchema_K2::PC_Struct;
+				PinType.PinSubCategoryObject = const_cast<UScriptStruct*>(Item->BlendableStructType.Get());
+				break;
+		}
+		return PinType;
+	}
+
+	void OnBlendableParameterPinTypeChanged(const FEdGraphPinType& PinType)
+	{
+		bool bIsValidType = true;
+		ECameraVariableType NewParameterType = ECameraVariableType::Boolean;
+
+		if (PinType.PinCategory == UEdGraphSchema_K2::PC_Boolean)
+		{
+			NewParameterType = ECameraVariableType::Boolean;
+		}
+		else if (PinType.PinCategory == UEdGraphSchema_K2::PC_Int)
+		{
+			NewParameterType = ECameraVariableType::Integer32;
+		}
+		else if (PinType.PinCategory == UEdGraphSchema_K2::PC_Float)
+		{
+			NewParameterType = ECameraVariableType::Float;
+		}
+		else if (PinType.PinCategory == UEdGraphSchema_K2::PC_Double)
+		{
+			NewParameterType = ECameraVariableType::Double;
+		}
+		else if (PinType.PinCategory == UEdGraphSchema_K2::PC_Struct)
+		{
+			const UObject* TypeObject = PinType.PinSubCategoryObject.Get();
+			if (TypeObject == TVariantStructure<FVector2f>::Get())
+			{
+				NewParameterType = ECameraVariableType::Vector2f;
+			}
+			else if (TypeObject == TBaseStructure<FVector2D>::Get())
+			{
+				NewParameterType = ECameraVariableType::Vector2d;
+			}
+			else if (TypeObject == TVariantStructure<FVector3f>::Get())
+			{
+				NewParameterType = ECameraVariableType::Vector3f;
+			}
+			else if (TypeObject == TBaseStructure<FVector>::Get())
+			{
+				NewParameterType = ECameraVariableType::Vector3d;
+			}
+			else if (TypeObject == TVariantStructure<FVector4f>::Get())
+			{
+				NewParameterType = ECameraVariableType::Vector4f;
+			}
+			else if (TypeObject == TBaseStructure<FVector4>::Get())
+			{
+				NewParameterType = ECameraVariableType::Vector4d;
+			}
+			else if (TypeObject == TVariantStructure<FRotator3f>::Get())
+			{
+				NewParameterType = ECameraVariableType::Rotator3f;
+			}
+			else if (TypeObject == TBaseStructure<FRotator>::Get())
+			{
+				NewParameterType = ECameraVariableType::Rotator3d;
+			}
+			else if (TypeObject == TVariantStructure<FTransform3f>::Get())
+			{
+				NewParameterType = ECameraVariableType::Transform3f;
+			}
+			else if (TypeObject == TBaseStructure<FTransform>::Get())
+			{
+				NewParameterType = ECameraVariableType::Transform3d;
+			}
+			else
+			{
+				bIsValidType = false;
+			}
+		}
+		else
+		{
+			bIsValidType = false;
+		}
+
+		if (ensure(bIsValidType) && Item->ParameterType != NewParameterType)
+		{
+			const FScopedTransaction Transaction(LOCTEXT("ChangeBlendableParameterType", "Change Blendable Parameter Type"));
+
+			Item->Modify();
+			Item->ParameterType = NewParameterType;
+		}
+	}
+
+	bool IsPinTypeSelectorReadOnly() const
+	{
+		return !IsHovered();
+	}
+};
+
+class FDataParameterPinTypeSelectorFilter : public IPinTypeSelectorFilter
+{
+	virtual bool ShouldShowPinTypeTreeItem(FPinTypeTreeItem InItem) const override
+	{
+		// Remove types that are for blendable parameters.
+		FEdGraphPinType PinType = InItem->GetPinType(false);
+		if (PinType.PinCategory == UEdGraphSchema_K2::PC_Boolean ||
+				PinType.PinCategory == UEdGraphSchema_K2::PC_Int ||
+				PinType.PinCategory == UEdGraphSchema_K2::PC_Int64 ||
+				PinType.PinCategory == UEdGraphSchema_K2::PC_Float ||
+				PinType.PinCategory == UEdGraphSchema_K2::PC_Double ||
+				PinType.PinCategory == UEdGraphSchema_K2::PC_Real)
+		{
+			return false;
+		}
+
+		if (PinType.PinCategory == UEdGraphSchema_K2::PC_Byte && PinType.PinSubCategoryObject == nullptr)
+		{
+			return false;
+		}
+
+		if (PinType.PinCategory == UEdGraphSchema_K2::PC_Struct)
+		{
+			const UObject* TypeObject = PinType.PinSubCategoryObject.Get();
+			if (TypeObject == TVariantStructure<FVector2f>::Get() ||
+					TypeObject == TBaseStructure<FVector2D>::Get() ||
+					TypeObject == TVariantStructure<FVector3f>::Get() ||
+					TypeObject == TBaseStructure<FVector>::Get() ||
+					TypeObject == TVariantStructure<FVector4f>::Get() ||
+					TypeObject == TBaseStructure<FVector4>::Get() ||
+					TypeObject == TVariantStructure<FRotator3f>::Get() ||
+					TypeObject == TBaseStructure<FRotator>::Get() ||
+					TypeObject == TVariantStructure<FTransform3f>::Get() ||
+					TypeObject == TBaseStructure<FTransform>::Get())
+			{
+				return false;
+			}
+		}
+		
+		return true;
+	};
+};
+
+/**
+ * List entry for the data parameters panel.
+ */
+class SCameraObjectInterfaceDataParameterTableRow : public SCameraObjectInterfaceParameterTableRowBase<UCameraObjectInterfaceDataParameter>
+{
+protected:
+
+	virtual TSharedRef<SWidget> GenerateWidgetForColumn(const FName& InColumnName) override
+	{
+		if (InColumnName == ParameterTypeColumn)
+		{
+			TSharedRef<FGameplayCamerasEditorStyle> GameplayCamerasEditorStyle = FGameplayCamerasEditorStyle::Get();
+
+			const UEdGraphSchema_K2* K2Schema = GetDefault<UEdGraphSchema_CameraNodeK2>();
+			TArray<TSharedPtr<IPinTypeSelectorFilter>> PinTypeSelectorFilters{ MakeShared<FDataParameterPinTypeSelectorFilter>() };
+
+			return SNew(SBox)
+				.MinDesiredWidth(125.f)
+				.Padding(this, &SCameraObjectInterfaceDataParameterTableRow::GetPinTypeSelectorPadding)
+				[
+					SNew(SPinTypeSelector, FGetPinTypeTree::CreateUObject(K2Schema, &UEdGraphSchema_K2::GetVariableTypeTree))
+					.OnPinTypeChanged(this, &SCameraObjectInterfaceDataParameterTableRow::OnDataParameterPinTypeChanged)
+					.TargetPinType(this, &SCameraObjectInterfaceDataParameterTableRow::GetDataParameterPinType)
+					.ReadOnly(this, &SCameraObjectInterfaceDataParameterTableRow::IsPinTypeSelectorReadOnly)
+					.Schema(K2Schema)
+					.CustomFilters(PinTypeSelectorFilters)
+					.bAllowArrays(true)
+				];
+		}
+
+		return SCameraObjectInterfaceParameterTableRowBase<UCameraObjectInterfaceDataParameter>::GenerateWidgetForColumn(InColumnName);
+	}
+
+	FMargin GetPinTypeSelectorPadding() const
+	{
+		// Add some horizontal margin when the pin type selector is read-only, so that the transition on hover is seamless.
+		return IsHovered() ? FMargin(0) : FMargin(7, 0);
+	}
+
+	FEdGraphPinType GetDataParameterPinType() const
+	{
+		FEdGraphPinType PinType;
+		PinType.PinSubCategoryObject = const_cast<UObject*>(Item->DataTypeObject.Get());
+
+		switch (Item->DataType)
+		{
+			case ECameraContextDataType::Name:
+				PinType.PinCategory = UEdGraphSchema_K2::PC_Name;
+				break;
+			case ECameraContextDataType::String:
+				PinType.PinCategory = UEdGraphSchema_K2::PC_String;
+				break;
+			case ECameraContextDataType::Enum:
+				PinType.PinCategory = UEdGraphSchema_K2::PC_Enum;
+				break;
+			case ECameraContextDataType::Struct:
+				PinType.PinCategory = UEdGraphSchema_K2::PC_Struct;
+				break;
+			case ECameraContextDataType::Object:
+				PinType.PinCategory = UEdGraphSchema_K2::PC_Object;
+				break;
+			case ECameraContextDataType::Class:
+				PinType.PinCategory = UEdGraphSchema_K2::PC_Class;
+				break;
+		}
+
+		switch (Item->DataContainerType)
+		{
+			case ECameraContextDataContainerType::Array:
+				PinType.ContainerType = EPinContainerType::Array;
+				break;
+		}
+
+		return PinType;
+	}
+
+	void OnDataParameterPinTypeChanged(const FEdGraphPinType& PinType)
+	{
+		bool bIsValidType = true;
+		ECameraContextDataType NewDataType = ECameraContextDataType::Name;
+		ECameraContextDataContainerType NewDataContainerType = ECameraContextDataContainerType::None;
+		const UObject* NewDataTypeObject = PinType.PinSubCategoryObject.Get();
+
+		if (PinType.PinCategory == UEdGraphSchema_K2::PC_Name)
+		{
+			NewDataType = ECameraContextDataType::Name;
+		}
+		else if (PinType.PinCategory == UEdGraphSchema_K2::PC_String)
+		{
+			NewDataType = ECameraContextDataType::String;
+		}
+		else if (PinType.PinCategory == UEdGraphSchema_K2::PC_Enum)
+		{
+			NewDataType = ECameraContextDataType::Enum;
+		}
+		else if (PinType.PinCategory == UEdGraphSchema_K2::PC_Byte && 
+				PinType.PinSubCategoryObject != nullptr && PinType.PinSubCategoryObject->IsA<UEnum>())
+		{
+			NewDataType = ECameraContextDataType::Enum;
+		}
+		else if (PinType.PinCategory == UEdGraphSchema_K2::PC_Struct)
+		{
+			NewDataType = ECameraContextDataType::Struct;
+		}
+		else if (PinType.PinCategory == UEdGraphSchema_K2::PC_Object)
+		{
+			NewDataType = ECameraContextDataType::Object;
+		}
+		else if (PinType.PinCategory == UEdGraphSchema_K2::PC_Class)
+		{
+			NewDataType = ECameraContextDataType::Class;
+		}
+		else
+		{
+			bIsValidType = false;
+		}
+
+		switch (PinType.ContainerType)
+		{
+			case EPinContainerType::None:
+				NewDataContainerType = ECameraContextDataContainerType::None;
+				break;
+			case EPinContainerType::Array:
+				NewDataContainerType = ECameraContextDataContainerType::Array;
+				break;
+			default:
+				bIsValidType = false;
+				break;
+		}
+
+		if (ensure(bIsValidType) && (NewDataType != Item->DataType || NewDataContainerType != Item->DataContainerType || NewDataTypeObject != Item->DataTypeObject))
+		{
+			const FScopedTransaction Transaction(LOCTEXT("ChangeDataParameterType", "Change Data Parameter Type"));
+
+			Item->Modify();
+
+			Item->DataType = NewDataType;
+			Item->DataTypeObject = NewDataTypeObject;
+		}
+	}
+
+	bool IsPinTypeSelectorReadOnly() const
+	{
+		return !IsHovered();
+	}
+};
+
+void SCameraObjectInterfaceParametersPanel::Construct(const FArguments& Args, FCameraObjectInterfaceParametersToolkit* OwnerToolkit)
+{
+	CameraObject = OwnerToolkit->GetCameraObject();
+	CameraObject->EventHandlers.Register(EventHandler, this);
+
+	Toolkit = OwnerToolkit;
+
+	ChildSlot
+	[
+		SNew(SSplitter)
+		.Orientation(Orient_Vertical)
+
+		// The blendable parameters panel.
+		+SSplitter::Slot()
+		[
+			SNew(SVerticalBox)
+			+SVerticalBox::Slot()
+			.AutoHeight()
+			[
+				SNew(SBorder)
+				.BorderImage(FAppStyle::GetBrush("ToolPanel.GroupBorder"))
+				.Padding(FMargin(3, 5))
+				[
+					SNew(SHorizontalBox)
+					+SHorizontalBox::Slot()
+					.FillWidth(1.f)
+					.VAlign(VAlign_Center)
+					.Padding(5.f)
+					[
+						SNew(SRichTextBlock)
+						.Text(LOCTEXT("BlendableParametersPanelTitle", "Blendable Parameters"))
+						.TransformPolicy(ETextTransformPolicy::ToUpper)
+						.DecoratorStyleSet(&FAppStyle::Get())
+						.TextStyle(FAppStyle::Get(), "DetailsView.CategoryTextStyle")
+					]
+					+SHorizontalBox::Slot()
+					.AutoWidth()
+					.VAlign(VAlign_Center)
+					[
+						SNew(SButton)
+						.ButtonStyle(FAppStyle::Get(), "SimpleButton")
+						.ContentPadding(FMargin(1, 0))
+						.ToolTipText(LOCTEXT("AddBlendableToolTip", "Add a blendable parameter"))
+						.OnClicked(this, &SCameraObjectInterfaceParametersPanel::OnAddBlendableParameter)
+						[
+							SNew(SImage)
+							.Image(FAppStyle::Get().GetBrush("Icons.PlusCircle"))
+							.ColorAndOpacity(FSlateColor::UseForeground())
+						]
+					]
+				]
+			]
+			+SVerticalBox::Slot()
+			.FillHeight(1.f)
+			[
+				SAssignNew(BlendableParametersListView, SListView<TObjectPtr<UCameraObjectInterfaceBlendableParameter>>)
+				.ListItemsSource(&CameraObject->Interface.BlendableParameters)
+				.OnGenerateRow(this, &SCameraObjectInterfaceParametersPanel::OnGenerateBlendableParameterRow)
+				.OnSelectionChanged(this, &SCameraObjectInterfaceParametersPanel::OnBlendableSelectionChanged)
+				.OnContextMenuOpening(this, &SCameraObjectInterfaceParametersPanel::OnBlendableParameterContextMenuOpening)
+				.HeaderRow
+				(
+					SNew(SHeaderRow)
+
+					+SHeaderRow::Column(ParameterTypeColumn)
+					.FillWidth(0.3f)
+					.DefaultLabel(LOCTEXT("ParameterTypeColumnLabel", "Type"))
+
+					+SHeaderRow::Column(ParameterNameColumn)
+					.FillWidth(0.5f)
+					.DefaultLabel(LOCTEXT("ParameterNameColumnLabel", "Name"))
+				)
+			]
+		]
+
+		// The data parameters panel.
+		+SSplitter::Slot()
+		[
+			SNew(SVerticalBox)
+			+SVerticalBox::Slot()
+			.AutoHeight()
+			[
+				SNew(SBorder)
+				.BorderImage(FAppStyle::GetBrush("ToolPanel.GroupBorder"))
+				.Padding(FMargin(3, 5))
+				[
+					SNew(SHorizontalBox)
+					+SHorizontalBox::Slot()
+					.FillWidth(1.f)
+					.VAlign(VAlign_Center)
+					.Padding(5.f)
+					[
+						SNew(SRichTextBlock)
+						.Text(LOCTEXT("DataParametersPanelTitle", "Data Parameters"))
+						.TransformPolicy(ETextTransformPolicy::ToUpper)
+						.DecoratorStyleSet(&FAppStyle::Get())
+						.TextStyle(FAppStyle::Get(), "DetailsView.CategoryTextStyle")
+					]
+					+SHorizontalBox::Slot()
+					.AutoWidth()
+					.VAlign(VAlign_Center)
+					[
+						SNew(SButton)
+						.ButtonStyle(FAppStyle::Get(), "SimpleButton")
+						.ContentPadding(FMargin(1, 0))
+						.ToolTipText(LOCTEXT("AddDataParameterToolTip", "Add a data parameter"))
+						.OnClicked(this, &SCameraObjectInterfaceParametersPanel::OnAddDataParameter)
+						[
+							SNew(SImage)
+							.Image(FAppStyle::Get().GetBrush("Icons.PlusCircle"))
+							.ColorAndOpacity(FSlateColor::UseForeground())
+						]
+					]
+				]
+			]
+			+SVerticalBox::Slot()
+			.FillHeight(1.f)
+			[
+				SAssignNew(DataParametersListView, SListView<TObjectPtr<UCameraObjectInterfaceDataParameter>>)
+				.ListItemsSource(&CameraObject->Interface.DataParameters)
+				.OnGenerateRow(this, &SCameraObjectInterfaceParametersPanel::OnGenerateDataParameterRow)
+				.OnSelectionChanged(this, &SCameraObjectInterfaceParametersPanel::OnDataParameterSelectionChanged)
+				.OnContextMenuOpening(this, &SCameraObjectInterfaceParametersPanel::OnDataParameterContextMenuOpening)
+				.HeaderRow
+				(
+					SNew(SHeaderRow)
+
+					+SHeaderRow::Column(ParameterTypeColumn)
+					.FillWidth(0.3f)
+					.DefaultLabel(LOCTEXT("ParameterTypeColumnLabel", "Type"))
+
+					+SHeaderRow::Column(ParameterNameColumn)
+					.FillWidth(0.5f)
+					.DefaultLabel(LOCTEXT("ParameterNameColumnLabel", "Name"))
+				)
+			]
+		]
+	];
+}
+
+void SCameraObjectInterfaceParametersPanel::RequestListRefresh()
+{
+	bListRefreshRequested = true;
+}
+
+void SCameraObjectInterfaceParametersPanel::Tick(const FGeometry& AllottedGeometry, const double InCurrentTime, const float InDeltaTime)
+{
+	SCompoundWidget::Tick(AllottedGeometry, InCurrentTime, InDeltaTime);
+
+	if (bListRefreshRequested)
+	{
+		bListRefreshRequested = false;
+
+		BlendableParametersListView->RequestListRefresh();
+		DataParametersListView->RequestListRefresh();
+	}
+}
+
+TSharedRef<ITableRow> SCameraObjectInterfaceParametersPanel::OnGenerateBlendableParameterRow(TObjectPtr<UCameraObjectInterfaceBlendableParameter> Item, const TSharedRef<STableViewBase>& OwnerTable)
+{
+	return SNew(SCameraObjectInterfaceBlendableParameterTableRow, OwnerTable)
+		.Item(Item);
+}
+
+TSharedRef<ITableRow> SCameraObjectInterfaceParametersPanel::OnGenerateDataParameterRow(TObjectPtr<UCameraObjectInterfaceDataParameter> Item, const TSharedRef<STableViewBase>& OwnerTable)
+{
+	return SNew(SCameraObjectInterfaceDataParameterTableRow, OwnerTable)
+		.Item(Item);
+}
+
+void SCameraObjectInterfaceParametersPanel::OnBlendableSelectionChanged(TObjectPtr<UCameraObjectInterfaceBlendableParameter> Item, ESelectInfo::Type Type)
+{
+	Toolkit->OnInterfaceParameterSelected().Broadcast(Item);
+}
+
+void SCameraObjectInterfaceParametersPanel::OnDataParameterSelectionChanged(TObjectPtr<UCameraObjectInterfaceDataParameter> Item, ESelectInfo::Type Type)
+{
+	Toolkit->OnInterfaceParameterSelected().Broadcast(Item);
+}
+
+TSharedPtr<SWidget> SCameraObjectInterfaceParametersPanel::OnBlendableParameterContextMenuOpening()
+{
+	TArray<TObjectPtr<UCameraObjectInterfaceBlendableParameter>> SelectedItems = BlendableParametersListView->GetSelectedItems();
+	if (SelectedItems.Num() > 0)
+	{
+		return OnInterfaceParameterContextMenuOpening(BlendableParametersListView, SelectedItems[0]);
+	}
+	return nullptr;
+}
+
+TSharedPtr<SWidget> SCameraObjectInterfaceParametersPanel::OnDataParameterContextMenuOpening()
+{
+	TArray<TObjectPtr<UCameraObjectInterfaceDataParameter>> SelectedItems = DataParametersListView->GetSelectedItems();
+	if (SelectedItems.Num() > 0)
+	{
+		return OnInterfaceParameterContextMenuOpening(DataParametersListView, SelectedItems[0]);
+	}
+	return nullptr;
+}
+
+void SCameraObjectInterfaceParametersPanel::OnSearchInterfaceParameter(TObjectPtr<UCameraObjectInterfaceParameterBase> Item)
+{
+	Toolkit->OnSearchInterfaceParameterNodes().Broadcast(Item);
+}
+
+template<typename ItemType>
+TSharedPtr<SWidget> SCameraObjectInterfaceParametersPanel::OnInterfaceParameterContextMenuOpening(
+		TSharedPtr<SListView<TObjectPtr<ItemType>>> ListView, TObjectPtr<ItemType> Item)
+{
+	FMenuBuilder MenuBuilder(true, nullptr);
+
+	MenuBuilder.BeginSection(NAME_None, LOCTEXT("InterfaceParameterActionsSection", "Actions"));
+	{
+		TObjectPtr<UCameraObjectInterfaceParameterBase> ItemBase(Item);
+		MenuBuilder.AddMenuEntry(
+				LOCTEXT("SearchInterfaceParameterNodes", "Find References"),
+				LOCTEXT("SearchInterfaceParameterNodesToolTip", "Find references to this interface parameter"),
+				FSlateIcon(FAppStyle::GetAppStyleSetName(), "Icons.Search"),
+				FExecuteAction::CreateSP(this, &SCameraObjectInterfaceParametersPanel::OnSearchInterfaceParameter, ItemBase));
+	}
+	MenuBuilder.EndSection();
+
+	MenuBuilder.BeginSection(NAME_None, LOCTEXT("InterfaceParameterEditSection", "Edit Interface Parameter"));
+	{
+		MenuBuilder.AddMenuEntry(
+				LOCTEXT("RenameInterfaceParameter", "Rename"),
+				LOCTEXT("RenameInterfaceParameterToolTip", "Renames this interface parameter"),
+				FSlateIcon(FAppStyle::GetAppStyleSetName(), "GenericCommands.Rename"),
+				FExecuteAction::CreateSP(this, &SCameraObjectInterfaceParametersPanel::OnRenameInterfaceParameter, ListView, Item));
+		MenuBuilder.AddMenuEntry(
+				LOCTEXT("DeleteInterfaceParameter", "Delete"),
+				LOCTEXT("DeleteInterfaceParameterToolTip", "Deletes this interface parameter"),
+				FSlateIcon(FAppStyle::GetAppStyleSetName(), "GenericCommands.Delete"),
+				FExecuteAction::CreateSP(this, &SCameraObjectInterfaceParametersPanel::OnDeleteInterfaceParameter, ListView, Item));
+	}
+	MenuBuilder.EndSection();
+
+	return MenuBuilder.MakeWidget();
+}
+
+template<typename ItemType>
+void SCameraObjectInterfaceParametersPanel::OnRenameInterfaceParameter(
+		TSharedPtr<SListView<TObjectPtr<ItemType>>> ListView, TObjectPtr<ItemType> Item)
+{
+	TSharedPtr<ITableRow> RowWidget = ListView->WidgetFromItem(Item);
+	if (!ensure(RowWidget))
+	{
+		return;
+	}
+
+	TSharedPtr<SCameraObjectInterfaceParameterTableRowBase<ItemType>> TypedRowWidget = 
+		StaticCastSharedPtr<SCameraObjectInterfaceParameterTableRowBase<ItemType>>(RowWidget);
+	TypedRowWidget->EnterNameEditingMode();
+}
+
+template<typename ItemType>
+void SCameraObjectInterfaceParametersPanel::OnDeleteInterfaceParameter(
+		TSharedPtr<SListView<TObjectPtr<ItemType>>> ListView, TObjectPtr<ItemType> Item)
+{
+	const FScopedTransaction Transaction(LOCTEXT("RemoveInterfaceParameter", "Remove Interface Parameter"));
+
+	CameraObject->Modify();
+
+	if constexpr (std::is_same_v<ItemType, UCameraObjectInterfaceBlendableParameter>)
+	{
+		const int32 NumRemoved = CameraObject->Interface.BlendableParameters.Remove(Item);
+		ensure(NumRemoved == 1);
+	}
+	else if constexpr (std::is_same_v<ItemType, UCameraObjectInterfaceDataParameter>)
+	{
+		const int32 NumRemoved = CameraObject->Interface.DataParameters.Remove(Item);
+		ensure(NumRemoved == 1);
+	}
+
+	CameraObject->EventHandlers.Notify(&ICameraObjectEventHandler::OnCameraObjectInterfaceChanged);
+
+	ListView->RequestListRefresh();
+}
+
+FReply SCameraObjectInterfaceParametersPanel::OnAddBlendableParameter()
+{
+	const FScopedTransaction Transaction(LOCTEXT("AddBlendableParameter", "Add Blendable Parameter"));
+
+	UCameraObjectInterfaceBlendableParameter* NewBlendableParameter = NewObject<UCameraObjectInterfaceBlendableParameter>(CameraObject, NAME_None, RF_Transactional);
+	NewBlendableParameter->InterfaceParameterName = NewBlendableParameter->GetName();
+
+	CameraObject->Modify();
+	CameraObject->Interface.BlendableParameters.Add(NewBlendableParameter);
+	CameraObject->EventHandlers.Notify(&ICameraObjectEventHandler::OnCameraObjectInterfaceChanged);
+
+	return FReply::Handled();
+}
+
+FReply SCameraObjectInterfaceParametersPanel::OnAddDataParameter()
+{
+	const FScopedTransaction Transaction(LOCTEXT("AddDataParameter", "Add Data Parameter"));
+
+	UCameraObjectInterfaceDataParameter* NewDataParameter = NewObject<UCameraObjectInterfaceDataParameter>(CameraObject, NAME_None, RF_Transactional);
+	NewDataParameter->InterfaceParameterName = NewDataParameter->GetName();
+
+	CameraObject->Modify();
+	CameraObject->Interface.DataParameters.Add(NewDataParameter);
+	CameraObject->EventHandlers.Notify(&ICameraObjectEventHandler::OnCameraObjectInterfaceChanged);
+
+	return FReply::Handled();
+}
+
+void SCameraObjectInterfaceParametersPanel::OnCameraObjectInterfaceChanged()
+{
+	bListRefreshRequested = true;
+}
+
+}  // namespace UE::Cameras
+
+#undef LOCTEXT_NAMESPACE
+
diff --git a/Engine/Plugins/Cameras/GameplayCameras/Source/GameplayCamerasEditor/Private/Editors/SCameraRigAssetEditor.cpp b/Engine/Plugins/Cameras/GameplayCameras/Source/GameplayCamerasEditor/Private/Editors/SCameraRigAssetEditor.cpp
--- a/Engine/Plugins/Cameras/GameplayCameras/Source/GameplayCamerasEditor/Private/Editors/SCameraRigAssetEditor.cpp
+++ b/Engine/Plugins/Cameras/GameplayCameras/Source/GameplayCamerasEditor/Private/Editors/SCameraRigAssetEditor.cpp
@@ -29,8 +29,6 @@ void SCameraRigAssetEditor::Construct(const FArguments& InArgs)
 
 	CreateGraphEditors();
 
-	CameraRigAsset->UBaseCameraObject::EventHandlers.Register(EventHandler, this);
-
 	ChildSlot
 	[
 		SAssignNew(BoxPanel, SBox)
@@ -52,8 +50,6 @@ void SCameraRigAssetEditor::SetCameraRigAsset(UCameraRigAsset* InCameraRig)
 {
 	if (CameraRigAsset != InCameraRig)
 	{
-		EventHandler.Unlink();
-
 		DiscardGraphEditors();
 
 		CameraRigAsset = InCameraRig;
@@ -61,8 +57,6 @@ void SCameraRigAssetEditor::SetCameraRigAsset(UCameraRigAsset* InCameraRig)
 		CreateGraphEditors();
 
 		SetEditorModeImpl(CurrentMode, true);
-
-		CameraRigAsset->UBaseCameraObject::EventHandlers.Register(EventHandler, this);
 	}
 }
 
@@ -332,53 +326,6 @@ void SCameraRigAssetEditor::RemoveOnAnyGraphChanged(FDelegateUserObjectConst InU
 	OnAnyGraphChanged.RemoveAll(InUserObject);
 }
 
-void SCameraRigAssetEditor::OnCameraObjectInterfaceChanged()
-{
-	// List all the interface parameters that want a node.
-	TSet<UCameraObjectInterfaceParameterBase*> InterfaceParametersWithNodes;
-	for (UCameraObjectInterfaceParameterBase* InterfaceParameter : CameraRigAsset->Interface.BlendableParameters)
-	{
-		if (InterfaceParameter->bHasGraphNode)
-		{
-			InterfaceParametersWithNodes.Add(InterfaceParameter);
-		}
-	}
-	for (UCameraObjectInterfaceParameterBase* InterfaceParameter : CameraRigAsset->Interface.DataParameters)
-	{
-		if (InterfaceParameter->bHasGraphNode)
-		{
-			InterfaceParametersWithNodes.Add(InterfaceParameter);
-		}
-	}
-
-	// Find all the interface parameter nodes that already exist.
-	TArray<UCameraObjectInterfaceParameterGraphNode*> InterfaceParameterNodes;
-	NodeGraph->GetNodesOfClass(InterfaceParameterNodes);
-	TMap<UCameraObjectInterfaceParameterBase*, UCameraObjectInterfaceParameterGraphNode*> InterfaceParameterToNodeMap;
-
-	// Remove nodes that aren't needed anymore.
-	for (UCameraObjectInterfaceParameterGraphNode* InterfaceParameterNode : InterfaceParameterNodes)
-	{
-		InterfaceParameterToNodeMap.Add(InterfaceParameterNode->GetInterfaceParameter(), InterfaceParameterNode);
-		if (!InterfaceParametersWithNodes.Contains(InterfaceParameterNode->GetInterfaceParameter()))
-		{
-			NodeGraph->Modify();
-			NodeGraph->RemoveNode(InterfaceParameterNode);
-		}
-	}
-
-	// Add nodes that are newly needed.
-	const UCameraRigCameraNodeGraphSchema* Schema = CastChecked<UCameraRigCameraNodeGraphSchema>(NodeGraph->GetSchema());
-	for (UCameraObjectInterfaceParameterBase* InterfaceParameter : InterfaceParametersWithNodes)
-	{
-		if (!InterfaceParameterToNodeMap.Contains(InterfaceParameter))
-		{
-			NodeGraph->Modify();
-			Schema->CreateInterfaceParameterNode(NodeGraph, InterfaceParameter);
-		}
-	}
-}
-
 }  // namespace UE::Cameras
 
 #undef LOCTEXT_NAMESPACE
diff --git a/Engine/Plugins/Cameras/GameplayCameras/Source/GameplayCamerasEditor/Private/GameplayCamerasEditorModule.cpp b/Engine/Plugins/Cameras/GameplayCameras/Source/GameplayCamerasEditor/Private/GameplayCamerasEditorModule.cpp
--- a/Engine/Plugins/Cameras/GameplayCameras/Source/GameplayCamerasEditor/Private/GameplayCamerasEditorModule.cpp
+++ b/Engine/Plugins/Cameras/GameplayCameras/Source/GameplayCamerasEditor/Private/GameplayCamerasEditorModule.cpp
@@ -19,6 +19,7 @@
 #include "ComponentVisualizers/GameplayCameraComponentVisualizer.h"
 #include "Core/CameraVariableCollection.h"
 #include "Customizations/CameraAssetReferenceDetailsCustomization.h"
+#include "Customizations/CameraObjectInterfaceParameterDetailsCustomization.h"
 #include "Customizations/CameraParameterDetailsCustomizations.h"
 #include "Customizations/CameraRigAssetReferenceDetailsCustomization.h"
 #include "Customizations/CameraShakeAssetReferenceDetailsCustomization.h"
@@ -492,6 +493,10 @@ private:
 				"FilmbackCameraNode", 
 				FOnGetDetailCustomizationInstance::CreateStatic(
 					&FFilmbackCameraNodeDetailsCustomization::MakeInstance));
+		PropertyEditorModule.RegisterCustomClassLayout(
+				"CameraObjectInterfaceParameterBase",
+				FOnGetDetailCustomizationInstance::CreateStatic(
+					&FCameraObjectInterfaceParameterDetailsCustomization::MakeInstance));
 	}
 
 	void UnregisterDetailsCustomizations()
@@ -513,6 +518,7 @@ private:
 			PropertyEditorModule->UnregisterCustomPropertyTypeLayout("CameraShakeAssetReference");
 
 			PropertyEditorModule->UnregisterCustomClassLayout("FilmbackCameraNode");
+			PropertyEditorModule->UnregisterCustomClassLayout("CameraObjectInterfaceParameterBase");
 		}
 	}
 
diff --git a/Engine/Plugins/Cameras/GameplayCameras/Source/GameplayCamerasEditor/Private/Toolkits/CameraAssetEditorToolkit.cpp b/Engine/Plugins/Cameras/GameplayCameras/Source/GameplayCamerasEditor/Private/Toolkits/CameraAssetEditorToolkit.cpp
--- a/Engine/Plugins/Cameras/GameplayCameras/Source/GameplayCamerasEditor/Private/Toolkits/CameraAssetEditorToolkit.cpp
+++ b/Engine/Plugins/Cameras/GameplayCameras/Source/GameplayCamerasEditor/Private/Toolkits/CameraAssetEditorToolkit.cpp
@@ -366,6 +366,8 @@ void FCameraAssetEditorToolkit::OnBuild()
 		TabManager->TryInvokeTab(MessagesTabId);
 	}
 
+	DetailsView->RequestForceRefresh();
+
 	for (UCameraRigAsset* CameraRigAsset : UsageInfo.CameraRigs)
 	{
 		FCameraRigPackages BuiltPackages;
diff --git a/Engine/Plugins/Cameras/GameplayCameras/Source/GameplayCamerasEditor/Private/Toolkits/CameraObjectInterfaceParametersToolkit.cpp b/Engine/Plugins/Cameras/GameplayCameras/Source/GameplayCamerasEditor/Private/Toolkits/CameraObjectInterfaceParametersToolkit.cpp
--- a/Engine/Plugins/Cameras/GameplayCameras/Source/GameplayCamerasEditor/Private/Toolkits/CameraObjectInterfaceParametersToolkit.cpp
+++ b/Engine/Plugins/Cameras/GameplayCameras/Source/GameplayCamerasEditor/Private/Toolkits/CameraObjectInterfaceParametersToolkit.cpp
@@ -5,967 +5,16 @@
 #include "Core/BaseCameraObject.h"
 #include "EdGraph/EdGraphNode.h"
 #include "Editor.h"
-#include "Editors/CameraNodeGraphDragDropOp.h"
-#include "Framework/MultiBox/MultiBoxBuilder.h"
-#include "IGameplayCamerasModule.h"
-#include "ScopedTransaction.h"
+#include "Editors/SCameraObjectInterfaceParametersPanel.h"
 #include "SPinTypeSelector.h"
-#include "Styles/GameplayCamerasEditorStyle.h"
-#include "Styling/SlateIconFinder.h"
 #include "ToolMenus.h"
-#include "Types/SlateEnums.h"
-#include "Widgets/Input/SButton.h"
-#include "Widgets/Input/SCheckBox.h"
-#include "Widgets/Input/SComboButton.h"
 #include "Widgets/Layout/SBox.h"
-#include "Widgets/Layout/SSplitter.h"
-#include "Widgets/SBoxPanel.h"
-#include "Widgets/SCompoundWidget.h"
-#include "Widgets/Text/SInlineEditableTextBlock.h"
-#include "Widgets/Text/SRichTextBlock.h"
-#include "Widgets/Text/STextBlock.h"
-#include "Widgets/Views/SListView.h"
-
-#include UE_INLINE_GENERATED_CPP_BY_NAME(CameraObjectInterfaceParametersToolkit)
 
 #define LOCTEXT_NAMESPACE "CameraObjectInterfaceParametersToolkit"
 
 namespace UE::Cameras
 {
 
-static const FName ParameterTypeColumn(TEXT("ParameterType"));
-static const FName ParameterNameColumn(TEXT("ParameterName"));
-static const FName ParameterMessageColumn(TEXT("ParameterMessage"));
-static const FName ParameterIsPreBlendedColumn(TEXT("ParameterIsPreBlended"));
-
-/**
- * List entry for any interface parameter panel.
- */
-template<typename ParameterType>
-class SCameraObjectInterfaceParameterTableRowBase : public SMultiColumnTableRow<TObjectPtr<ParameterType>>
-{
-public:
-
-	SLATE_BEGIN_ARGS(SCameraObjectInterfaceParameterTableRowBase<ParameterType>)
-	{}
-		SLATE_ARGUMENT(TObjectPtr<ParameterType>, Item)
-	SLATE_END_ARGS()
-
-	typedef SMultiColumnTableRow<TObjectPtr<ParameterType>> FSuperRowType;
-	typedef typename STableRow<TObjectPtr<ParameterType>>::FArguments FTableRowArgs;
-		
-	void Construct(const FArguments& Args, const TSharedRef<STableViewBase>& OwnerTable)
-	{
-		Item = Args._Item;
-
-		FSuperRowType::Construct(
-				FTableRowArgs(),
-				OwnerTable);
-	}
-
-	void EnterNameEditingMode()
-	{
-		NameTextBlock->EnterEditingMode();
-	}
-
-protected:
-
-	virtual TSharedRef<SWidget> GenerateWidgetForColumn(const FName& InColumnName) override
-	{
-		if (InColumnName == ParameterNameColumn)
-		{
-			return SAssignNew(NameTextBlock, SInlineEditableTextBlock)
-				.IsSelected(this, &SCameraObjectInterfaceParameterTableRowBase<ParameterType>::IsSelected)
-				.Text_Lambda([this]() { return FText::FromString(Item->InterfaceParameterName); })
-				.OnTextCommitted(this, &SCameraObjectInterfaceParameterTableRowBase<ParameterType>::OnParameterNameTextCommitted)
-				.ToolTipText(this, &SCameraObjectInterfaceParameterTableRowBase<ParameterType>::GetParameterToolTip);
-		}
-		else if (InColumnName == ParameterMessageColumn)
-		{
-			TSharedRef<FGameplayCamerasEditorStyle> GameplayCamerasEditorStyle = FGameplayCamerasEditorStyle::Get();
-			const FTextBlockStyle& MessageStyle = GameplayCamerasEditorStyle->GetWidgetStyle<FTextBlockStyle>("CameraObjectEditor.InterfaceParameter.Message");
-
-			return SNew(STextBlock)
-				.TextStyle(&MessageStyle)
-				.Text(this, &SCameraObjectInterfaceParameterTableRowBase<ParameterType>::GetParameterMessageText)
-				.ToolTipText(this, &SCameraObjectInterfaceParameterTableRowBase<ParameterType>::GetParameterMessageToolTip);
-		}
-
-		return SNullWidget::NullWidget;
-	}
-
-	virtual FReply OnDragDetected(const FGeometry& MyGeometry, const FPointerEvent& MouseEvent) override
-	{
-		if (MouseEvent.IsMouseButtonDown(EKeys::LeftMouseButton))
-		{
-			return FReply::Handled().BeginDragDrop(FCameraNodeGraphInterfaceParameterDragDropOp::New(Item));
-		}
-		return SMultiColumnTableRow<TObjectPtr<ParameterType>>::OnDragDetected(MyGeometry, MouseEvent);
-	}
-
-	void OnParameterNameTextCommitted(const FText& Text, ETextCommit::Type CommitType)
-	{
-		const FScopedTransaction Transaction(LOCTEXT("RenameInterfaceParameterTransaction", "Rename Interface Parameter"));
-
-		Item->Modify();
-		Item->InterfaceParameterName = Text.ToString();
-	}
-
-	FText GetParameterToolTip() const
-	{
-		FGuid Guid = Item->GetGuid();
-		return FText::FromString(Guid.ToString(EGuidFormats::DigitsWithHyphensInParentheses));
-	}
-
-	FText GetParameterMessageText() const
-	{
-		if (!Item->Target || Item->TargetPropertyName.IsNone())
-		{
-			return LOCTEXT("UnboundInterfaceParameterMessage", "Unbound");
-		}
-		return FText::GetEmpty();
-	}
-
-	FText GetParameterMessageToolTip() const
-	{
-		if (!Item->Target || Item->TargetPropertyName.IsNone())
-		{
-			return LOCTEXT(
-					"UnboundInterfaceParameterMessageToolTip", 
-					"This interface parameter is not bound to any camera node. Setting this parameter will have no effect.");
-		}
-		return FText::GetEmpty();
-	}
-
-protected:
-
-	TObjectPtr<ParameterType> Item;
-	TSharedPtr<SInlineEditableTextBlock> NameTextBlock;
-};
-
-/**
- * List entry for the blendable parameters panel.
- */
-class SCameraObjectInterfaceBlendableParameterTableRow : public SCameraObjectInterfaceParameterTableRowBase<UCameraObjectInterfaceBlendableParameter>
-{
-public:
-
-	void Construct(const FArguments& Args, const TSharedRef<STableViewBase>& OwnerTable)
-	{
-		SCameraObjectInterfaceParameterTableRowBase<UCameraObjectInterfaceBlendableParameter>::Construct(Args, OwnerTable);
-	}
-
-protected:
-
-	virtual TSharedRef<SWidget> GenerateWidgetForColumn(const FName& InColumnName) override
-	{
-		if (InColumnName == ParameterTypeColumn)
-		{
-			TSharedRef<FGameplayCamerasEditorStyle> GameplayCamerasEditorStyle = FGameplayCamerasEditorStyle::Get();
-
-			const UEdGraphSchema_K2* K2Schema = GetDefault<UEdGraphSchema_CameraNodeK2>();
-
-			return SNew(SBox)
-				.MinDesiredWidth(125.f)
-				.Padding(this, &SCameraObjectInterfaceBlendableParameterTableRow::GetPinTypeSelectorPadding)
-				[
-					SNew(SPinTypeSelector, FGetPinTypeTree::CreateSP(this, &SCameraObjectInterfaceBlendableParameterTableRow::GetPinTypeTreeInfos))
-					.OnPinTypeChanged(this, &SCameraObjectInterfaceBlendableParameterTableRow::OnBlendableParameterPinTypeChanged)
-					.TargetPinType(this, &SCameraObjectInterfaceBlendableParameterTableRow::GetBlendableParameterPinType)
-					.ReadOnly(this, &SCameraObjectInterfaceBlendableParameterTableRow::IsPinTypeSelectorReadOnly)
-					.Schema(K2Schema)
-					.bAllowArrays(false)
-				];
-		}
-		else if (InColumnName == ParameterIsPreBlendedColumn)
-		{
-			return SNew(SCheckBox)
-				.IsChecked(this, &SCameraObjectInterfaceBlendableParameterTableRow::IsBlendableParameterPreBlended)
-				.OnCheckStateChanged(this, &SCameraObjectInterfaceBlendableParameterTableRow::OnBlendableParameterPreBlendedChanged);
-		}
-
-		return SCameraObjectInterfaceParameterTableRowBase<UCameraObjectInterfaceBlendableParameter>::GenerateWidgetForColumn(InColumnName);
-	}
-
-	FMargin GetPinTypeSelectorPadding() const
-	{
-		// Add some horizontal margin when the pin type selector is read-only, so that the transition on hover is seamless.
-		return IsHovered() ? FMargin(0) : FMargin(7, 0);
-	}
-
-	void GetPinTypeTreeInfos(TArray<FPinTypeTreeItem>& TypeTree, ETypeTreeFilter TypeTreeFilter) const
-	{
-		using FPinTypeTreeInfo = UEdGraphSchema_K2::FPinTypeTreeInfo;
-
-		const UEdGraphSchema_K2* K2Schema = GetDefault<UEdGraphSchema_CameraNodeK2>();
-
-		TypeTree.Reset();
-
-		TypeTree.Add(MakeShared<FPinTypeTreeInfo>(UEdGraphSchema_K2::GetCategoryText(UEdGraphSchema_K2::PC_Boolean, true), UEdGraphSchema_K2::PC_Boolean, K2Schema, LOCTEXT("BooleanType", "True or false value")));
-		TypeTree.Add(MakeShared<FPinTypeTreeInfo>(UEdGraphSchema_K2::GetCategoryText(UEdGraphSchema_K2::PC_Int, true), UEdGraphSchema_K2::PC_Int, K2Schema, LOCTEXT("IntegerType", "Integer number")));
-		TypeTree.Add(MakeShared<FPinTypeTreeInfo>(UEdGraphSchema_K2::GetCategoryText(UEdGraphSchema_K2::PC_Float, true), UEdGraphSchema_K2::PC_Float, K2Schema, LOCTEXT("FloatType", "Floating point number")));
-		TypeTree.Add(MakeShared<FPinTypeTreeInfo>(UEdGraphSchema_K2::GetCategoryText(UEdGraphSchema_K2::PC_Double, true), UEdGraphSchema_K2::PC_Double, K2Schema, LOCTEXT("DoubleType", "Double precision floating point number")));
-
-		TypeTree.Add(MakeShared<FPinTypeTreeInfo>(UEdGraphSchema_K2::PC_Struct, TVariantStructure<FVector2f>::Get(), LOCTEXT("Vector2fType", "A 2D vector")));
-		TypeTree.Add(MakeShared<FPinTypeTreeInfo>(UEdGraphSchema_K2::PC_Struct, TBaseStructure<FVector2D>::Get(), LOCTEXT("Vector2dType", "A double precision 2D vector")));
-		TypeTree.Add(MakeShared<FPinTypeTreeInfo>(UEdGraphSchema_K2::PC_Struct, TVariantStructure<FVector3f>::Get(), LOCTEXT("Vector3fType", "A 3D vector")));
-		TypeTree.Add(MakeShared<FPinTypeTreeInfo>(UEdGraphSchema_K2::PC_Struct, TBaseStructure<FVector>::Get(), LOCTEXT("Vector3dType", "A double precision 3D vector")));
-		TypeTree.Add(MakeShared<FPinTypeTreeInfo>(UEdGraphSchema_K2::PC_Struct, TVariantStructure<FVector4f>::Get(), LOCTEXT("Vector4fType", "A 4D vector")));
-		TypeTree.Add(MakeShared<FPinTypeTreeInfo>(UEdGraphSchema_K2::PC_Struct, TBaseStructure<FVector4>::Get(), LOCTEXT("Vector4dType", "A double precision 4D vector")));
-		TypeTree.Add(MakeShared<FPinTypeTreeInfo>(UEdGraphSchema_K2::PC_Struct, TVariantStructure<FRotator3f>::Get(), LOCTEXT("Rotator3fType", "A 3D rotation")));
-		TypeTree.Add(MakeShared<FPinTypeTreeInfo>(UEdGraphSchema_K2::PC_Struct, TBaseStructure<FRotator>::Get(), LOCTEXT("Rotator3dType", "A double precision 3D rotation")));
-		TypeTree.Add(MakeShared<FPinTypeTreeInfo>(UEdGraphSchema_K2::PC_Struct, TVariantStructure<FTransform3f>::Get(), LOCTEXT("Transform3fType", "A 3D transformation")));
-		TypeTree.Add(MakeShared<FPinTypeTreeInfo>(UEdGraphSchema_K2::PC_Struct, TBaseStructure<FTransform>::Get(), LOCTEXT("Transform3dType", "A double precision 3D transformation")));
-
-		TSharedPtr<FPinTypeTreeInfo> Structs = MakeShared<FPinTypeTreeInfo>(
-				LOCTEXT("BlendableStructPinTypeLabel", "Blendable Structures"),
-				UEdGraphSchema_K2::PC_Struct,
-				K2Schema,
-				LOCTEXT("BlendableStructPinTypeToolTip", "Blendable structure types"),
-				true);
-		IGameplayCamerasModule& GameplayCamerasModule = IGameplayCamerasModule::Get();
-		for (const FBlendableStructInfo& BlendableStruct : GameplayCamerasModule.GetBlendableStructs())
-		{
-			if (const UScriptStruct* StructType = BlendableStruct.StructType)
-			{
-				Structs->Children.Add(
-						MakeShared<UEdGraphSchema_K2::FPinTypeTreeInfo>(
-							UEdGraphSchema_K2::PC_Struct,
-							const_cast<UScriptStruct*>(StructType),
-							StructType->GetToolTipText(),
-							false,
-							(uint8)EObjectReferenceType::NotAnObject));
-			}
-		}
-		TypeTree.Add(Structs);
-	}
-
-	FEdGraphPinType GetBlendableParameterPinType() const
-	{
-		FEdGraphPinType PinType;
-		switch (Item->ParameterType)
-		{
-			case ECameraVariableType::Boolean:
-				PinType.PinCategory = UEdGraphSchema_K2::PC_Boolean;
-				break;
-			case ECameraVariableType::Integer32:
-				PinType.PinCategory = UEdGraphSchema_K2::PC_Int;
-				break;
-			case ECameraVariableType::Float:
-				PinType.PinCategory = UEdGraphSchema_K2::PC_Float;
-				break;
-			case ECameraVariableType::Double:
-				PinType.PinCategory = UEdGraphSchema_K2::PC_Double;
-				break;
-			case ECameraVariableType::Vector2f:
-				PinType.PinCategory = UEdGraphSchema_K2::PC_Struct;
-				PinType.PinSubCategoryObject = TVariantStructure<FVector2f>::Get();
-				break;
-			case ECameraVariableType::Vector2d:
-				PinType.PinCategory = UEdGraphSchema_K2::PC_Struct;
-				PinType.PinSubCategoryObject = TBaseStructure<FVector2D>::Get();
-				break;
-			case ECameraVariableType::Vector3f:
-				PinType.PinCategory = UEdGraphSchema_K2::PC_Struct;
-				PinType.PinSubCategoryObject = TVariantStructure<FVector3f>::Get();
-				break;
-			case ECameraVariableType::Vector3d:
-				PinType.PinCategory = UEdGraphSchema_K2::PC_Struct;
-				PinType.PinSubCategoryObject = TBaseStructure<FVector>::Get();
-				break;
-			case ECameraVariableType::Vector4f:
-				PinType.PinCategory = UEdGraphSchema_K2::PC_Struct;
-				PinType.PinSubCategoryObject = TVariantStructure<FVector4f>::Get();
-				break;
-			case ECameraVariableType::Vector4d:
-				PinType.PinCategory = UEdGraphSchema_K2::PC_Struct;
-				PinType.PinSubCategoryObject = TBaseStructure<FVector4>::Get();
-				break;
-			case ECameraVariableType::Rotator3f:
-				PinType.PinCategory = UEdGraphSchema_K2::PC_Struct;
-				PinType.PinSubCategoryObject = TVariantStructure<FRotator3f>::Get();
-				break;
-			case ECameraVariableType::Rotator3d:
-				PinType.PinCategory = UEdGraphSchema_K2::PC_Struct;
-				PinType.PinSubCategoryObject = TBaseStructure<FRotator>::Get();
-				break;
-			case ECameraVariableType::Transform3f:
-				PinType.PinCategory = UEdGraphSchema_K2::PC_Struct;
-				PinType.PinSubCategoryObject = TVariantStructure<FTransform3f>::Get();
-				break;
-			case ECameraVariableType::Transform3d:
-				PinType.PinCategory = UEdGraphSchema_K2::PC_Struct;
-				PinType.PinSubCategoryObject = TBaseStructure<FTransform>::Get();
-				break;
-			case ECameraVariableType::BlendableStruct:
-				PinType.PinCategory = UEdGraphSchema_K2::PC_Struct;
-				PinType.PinSubCategoryObject = const_cast<UScriptStruct*>(Item->BlendableStructType.Get());
-				break;
-		}
-		return PinType;
-	}
-
-	void OnBlendableParameterPinTypeChanged(const FEdGraphPinType& PinType)
-	{
-		bool bIsValidType = true;
-		ECameraVariableType NewParameterType = ECameraVariableType::Boolean;
-
-		if (PinType.PinCategory == UEdGraphSchema_K2::PC_Boolean)
-		{
-			NewParameterType = ECameraVariableType::Boolean;
-		}
-		else if (PinType.PinCategory == UEdGraphSchema_K2::PC_Int)
-		{
-			NewParameterType = ECameraVariableType::Integer32;
-		}
-		else if (PinType.PinCategory == UEdGraphSchema_K2::PC_Float)
-		{
-			NewParameterType = ECameraVariableType::Float;
-		}
-		else if (PinType.PinCategory == UEdGraphSchema_K2::PC_Double)
-		{
-			NewParameterType = ECameraVariableType::Double;
-		}
-		else if (PinType.PinCategory == UEdGraphSchema_K2::PC_Struct)
-		{
-			const UObject* TypeObject = PinType.PinSubCategoryObject.Get();
-			if (TypeObject == TVariantStructure<FVector2f>::Get())
-			{
-				NewParameterType = ECameraVariableType::Vector2f;
-			}
-			else if (TypeObject == TBaseStructure<FVector2D>::Get())
-			{
-				NewParameterType = ECameraVariableType::Vector2d;
-			}
-			else if (TypeObject == TVariantStructure<FVector3f>::Get())
-			{
-				NewParameterType = ECameraVariableType::Vector3f;
-			}
-			else if (TypeObject == TBaseStructure<FVector>::Get())
-			{
-				NewParameterType = ECameraVariableType::Vector3d;
-			}
-			else if (TypeObject == TVariantStructure<FVector4f>::Get())
-			{
-				NewParameterType = ECameraVariableType::Vector4f;
-			}
-			else if (TypeObject == TBaseStructure<FVector4>::Get())
-			{
-				NewParameterType = ECameraVariableType::Vector4d;
-			}
-			else if (TypeObject == TVariantStructure<FRotator3f>::Get())
-			{
-				NewParameterType = ECameraVariableType::Rotator3f;
-			}
-			else if (TypeObject == TBaseStructure<FRotator>::Get())
-			{
-				NewParameterType = ECameraVariableType::Rotator3d;
-			}
-			else if (TypeObject == TVariantStructure<FTransform3f>::Get())
-			{
-				NewParameterType = ECameraVariableType::Transform3f;
-			}
-			else if (TypeObject == TBaseStructure<FTransform>::Get())
-			{
-				NewParameterType = ECameraVariableType::Transform3d;
-			}
-			else
-			{
-				bIsValidType = false;
-			}
-		}
-		else
-		{
-			bIsValidType = false;
-		}
-
-		if (ensure(bIsValidType) && Item->ParameterType != NewParameterType)
-		{
-			const FScopedTransaction Transaction(LOCTEXT("ChangeBlendableParameterType", "Change Blendable Parameter Type"));
-
-			Item->Modify();
-			Item->ParameterType = NewParameterType;
-		}
-	}
-
-	ECheckBoxState IsBlendableParameterPreBlended() const
-	{
-		return Item->bIsPreBlended ? ECheckBoxState::Checked : ECheckBoxState::Unchecked;
-	}
-	
-	void OnBlendableParameterPreBlendedChanged(ECheckBoxState CheckState)
-	{
-		const bool bIsPreBlended = (CheckState == ECheckBoxState::Checked);
-		if (Item->bIsPreBlended != bIsPreBlended)
-		{
-			const FScopedTransaction Transaction(LOCTEXT("ChangeBlendableParameterIsPreBlended", "Change Blendable Parameter Pre-Blending"));
-
-			Item->Modify();
-			Item->bIsPreBlended = bIsPreBlended;
-		}
-	}
-
-	bool IsPinTypeSelectorReadOnly() const
-	{
-		return !IsHovered();
-	}
-};
-
-class FDataParameterPinTypeSelectorFilter : public IPinTypeSelectorFilter
-{
-	virtual bool ShouldShowPinTypeTreeItem(FPinTypeTreeItem InItem) const override
-	{
-		// Remove types that are for blendable parameters.
-		FEdGraphPinType PinType = InItem->GetPinType(false);
-		if (PinType.PinCategory == UEdGraphSchema_K2::PC_Boolean ||
-				PinType.PinCategory == UEdGraphSchema_K2::PC_Int ||
-				PinType.PinCategory == UEdGraphSchema_K2::PC_Int64 ||
-				PinType.PinCategory == UEdGraphSchema_K2::PC_Float ||
-				PinType.PinCategory == UEdGraphSchema_K2::PC_Double ||
-				PinType.PinCategory == UEdGraphSchema_K2::PC_Real)
-		{
-			return false;
-		}
-
-		if (PinType.PinCategory == UEdGraphSchema_K2::PC_Byte && PinType.PinSubCategoryObject == nullptr)
-		{
-			return false;
-		}
-
-		if (PinType.PinCategory == UEdGraphSchema_K2::PC_Struct)
-		{
-			const UObject* TypeObject = PinType.PinSubCategoryObject.Get();
-			if (TypeObject == TVariantStructure<FVector2f>::Get() ||
-					TypeObject == TBaseStructure<FVector2D>::Get() ||
-					TypeObject == TVariantStructure<FVector3f>::Get() ||
-					TypeObject == TBaseStructure<FVector>::Get() ||
-					TypeObject == TVariantStructure<FVector4f>::Get() ||
-					TypeObject == TBaseStructure<FVector4>::Get() ||
-					TypeObject == TVariantStructure<FRotator3f>::Get() ||
-					TypeObject == TBaseStructure<FRotator>::Get() ||
-					TypeObject == TVariantStructure<FTransform3f>::Get() ||
-					TypeObject == TBaseStructure<FTransform>::Get())
-			{
-				return false;
-			}
-		}
-		
-		return true;
-	};
-};
-
-/**
- * List entry for the data parameters panel.
- */
-class SCameraObjectInterfaceDataParameterTableRow : public SCameraObjectInterfaceParameterTableRowBase<UCameraObjectInterfaceDataParameter>
-{
-protected:
-
-	virtual TSharedRef<SWidget> GenerateWidgetForColumn(const FName& InColumnName) override
-	{
-		if (InColumnName == ParameterTypeColumn)
-		{
-			TSharedRef<FGameplayCamerasEditorStyle> GameplayCamerasEditorStyle = FGameplayCamerasEditorStyle::Get();
-
-			const UEdGraphSchema_K2* K2Schema = GetDefault<UEdGraphSchema_CameraNodeK2>();
-			TArray<TSharedPtr<IPinTypeSelectorFilter>> PinTypeSelectorFilters{ MakeShared<FDataParameterPinTypeSelectorFilter>() };
-
-			return SNew(SBox)
-				.MinDesiredWidth(125.f)
-				.Padding(this, &SCameraObjectInterfaceDataParameterTableRow::GetPinTypeSelectorPadding)
-				[
-					SNew(SPinTypeSelector, FGetPinTypeTree::CreateUObject(K2Schema, &UEdGraphSchema_K2::GetVariableTypeTree))
-					.OnPinTypeChanged(this, &SCameraObjectInterfaceDataParameterTableRow::OnDataParameterPinTypeChanged)
-					.TargetPinType(this, &SCameraObjectInterfaceDataParameterTableRow::GetDataParameterPinType)
-					.ReadOnly(this, &SCameraObjectInterfaceDataParameterTableRow::IsPinTypeSelectorReadOnly)
-					.Schema(K2Schema)
-					.CustomFilters(PinTypeSelectorFilters)
-					.bAllowArrays(true)
-				];
-		}
-
-		return SCameraObjectInterfaceParameterTableRowBase<UCameraObjectInterfaceDataParameter>::GenerateWidgetForColumn(InColumnName);
-	}
-
-	FMargin GetPinTypeSelectorPadding() const
-	{
-		// Add some horizontal margin when the pin type selector is read-only, so that the transition on hover is seamless.
-		return IsHovered() ? FMargin(0) : FMargin(7, 0);
-	}
-
-	FEdGraphPinType GetDataParameterPinType() const
-	{
-		FEdGraphPinType PinType;
-		PinType.PinSubCategoryObject = const_cast<UObject*>(Item->DataTypeObject.Get());
-
-		switch (Item->DataType)
-		{
-			case ECameraContextDataType::Name:
-				PinType.PinCategory = UEdGraphSchema_K2::PC_Name;
-				break;
-			case ECameraContextDataType::String:
-				PinType.PinCategory = UEdGraphSchema_K2::PC_String;
-				break;
-			case ECameraContextDataType::Enum:
-				PinType.PinCategory = UEdGraphSchema_K2::PC_Enum;
-				break;
-			case ECameraContextDataType::Struct:
-				PinType.PinCategory = UEdGraphSchema_K2::PC_Struct;
-				break;
-			case ECameraContextDataType::Object:
-				PinType.PinCategory = UEdGraphSchema_K2::PC_Object;
-				break;
-			case ECameraContextDataType::Class:
-				PinType.PinCategory = UEdGraphSchema_K2::PC_Class;
-				break;
-		}
-
-		switch (Item->DataContainerType)
-		{
-			case ECameraContextDataContainerType::Array:
-				PinType.ContainerType = EPinContainerType::Array;
-				break;
-		}
-
-		return PinType;
-	}
-
-	void OnDataParameterPinTypeChanged(const FEdGraphPinType& PinType)
-	{
-		bool bIsValidType = true;
-		ECameraContextDataType NewDataType = ECameraContextDataType::Name;
-		ECameraContextDataContainerType NewDataContainerType = ECameraContextDataContainerType::None;
-		const UObject* NewDataTypeObject = PinType.PinSubCategoryObject.Get();
-
-		if (PinType.PinCategory == UEdGraphSchema_K2::PC_Name)
-		{
-			NewDataType = ECameraContextDataType::Name;
-		}
-		else if (PinType.PinCategory == UEdGraphSchema_K2::PC_String)
-		{
-			NewDataType = ECameraContextDataType::String;
-		}
-		else if (PinType.PinCategory == UEdGraphSchema_K2::PC_Enum)
-		{
-			NewDataType = ECameraContextDataType::Enum;
-		}
-		else if (PinType.PinCategory == UEdGraphSchema_K2::PC_Byte && 
-				PinType.PinSubCategoryObject != nullptr && PinType.PinSubCategoryObject->IsA<UEnum>())
-		{
-			NewDataType = ECameraContextDataType::Enum;
-		}
-		else if (PinType.PinCategory == UEdGraphSchema_K2::PC_Struct)
-		{
-			NewDataType = ECameraContextDataType::Struct;
-		}
-		else if (PinType.PinCategory == UEdGraphSchema_K2::PC_Object)
-		{
-			NewDataType = ECameraContextDataType::Object;
-		}
-		else if (PinType.PinCategory == UEdGraphSchema_K2::PC_Class)
-		{
-			NewDataType = ECameraContextDataType::Class;
-		}
-		else
-		{
-			bIsValidType = false;
-		}
-
-		switch (PinType.ContainerType)
-		{
-			case EPinContainerType::None:
-				NewDataContainerType = ECameraContextDataContainerType::None;
-				break;
-			case EPinContainerType::Array:
-				NewDataContainerType = ECameraContextDataContainerType::Array;
-				break;
-			default:
-				bIsValidType = false;
-				break;
-		}
-
-		if (ensure(bIsValidType) && (NewDataType != Item->DataType || NewDataContainerType != Item->DataContainerType || NewDataTypeObject != Item->DataTypeObject))
-		{
-			const FScopedTransaction Transaction(LOCTEXT("ChangeDataParameterType", "Change Data Parameter Type"));
-
-			Item->DataType = NewDataType;
-			Item->DataTypeObject = NewDataTypeObject;
-		}
-	}
-
-	bool IsPinTypeSelectorReadOnly() const
-	{
-		return !IsHovered();
-	}
-};
-
-/**
- * The overall interface parameters panel, showing two sub-panels, one for blendable parameters, and
- * one for data parameters.
- */
-class SCameraObjectInterfaceParametersPanel 
-	: public SCompoundWidget
-	, public ICameraObjectEventHandler
-{
-public:
-
-	SLATE_BEGIN_ARGS(SCameraObjectInterfaceParametersPanel)
-	{}
-	SLATE_END_ARGS()
-
-	void Construct(const FArguments& Args, FCameraObjectInterfaceParametersToolkit* OwnerToolkit);
-
-	void RequestListRefresh();
-
-protected:
-
-	// SWidget interface.
-	virtual void Tick(const FGeometry& AllottedGeometry, const double InCurrentTime, const float InDeltaTime) override;
-
-	// ICameraObjectEventHandler interface.
-	virtual void OnCameraObjectInterfaceChanged() override;
-
-private:
-
-	TSharedRef<ITableRow> OnGenerateBlendableParameterRow(TObjectPtr<UCameraObjectInterfaceBlendableParameter> Item, const TSharedRef<STableViewBase>& OwnerTable);
-	TSharedRef<ITableRow> OnGenerateDataParameterRow(TObjectPtr<UCameraObjectInterfaceDataParameter> Item, const TSharedRef<STableViewBase>& OwnerTable);
-
-	void OnBlendableSelectionChanged(TObjectPtr<UCameraObjectInterfaceBlendableParameter> Item, ESelectInfo::Type Type);
-	void OnDataParameterSelectionChanged(TObjectPtr<UCameraObjectInterfaceDataParameter> Item, ESelectInfo::Type Type);
-
-	TSharedPtr<SWidget> OnBlendableParameterContextMenuOpening();
-	TSharedPtr<SWidget> OnDataParameterContextMenuOpening();
-
-	template<typename ItemType>
-	TSharedPtr<SWidget> OnInterfaceParameterContextMenuOpening(
-			TSharedPtr<SListView<TObjectPtr<ItemType>>> ListView, TObjectPtr<ItemType> Item);
-	template<typename ItemType>
-	void OnRenameInterfaceParameter(
-			TSharedPtr<SListView<TObjectPtr<ItemType>>> ListView, TObjectPtr<ItemType> Item);
-	template<typename ItemType>
-	void OnDeleteInterfaceParameter(
-			TSharedPtr<SListView<TObjectPtr<ItemType>>> ListView, TObjectPtr<ItemType> Item);
-
-	FReply OnAddBlendableParameter();
-	FReply OnAddDataParameter();
-
-private:
-
-	UBaseCameraObject* CameraObject = nullptr;
-	FCameraObjectInterfaceParametersToolkit* Toolkit = nullptr;
-
-	TCameraEventHandler<ICameraObjectEventHandler> EventHandler;
-
-	TSharedPtr<SListView<TObjectPtr<UCameraObjectInterfaceBlendableParameter>>> BlendableParametersListView;
-	TSharedPtr<SListView<TObjectPtr<UCameraObjectInterfaceDataParameter>>> DataParametersListView;
-
-	bool bListRefreshRequested = false;
-};
-
-void SCameraObjectInterfaceParametersPanel::Construct(const FArguments& Args, FCameraObjectInterfaceParametersToolkit* OwnerToolkit)
-{
-	CameraObject = OwnerToolkit->GetCameraObject();
-	CameraObject->EventHandlers.Register(EventHandler, this);
-
-	Toolkit = OwnerToolkit;
-
-	ChildSlot
-	[
-		SNew(SSplitter)
-		.Orientation(Orient_Vertical)
-
-		// The blendable parameters panel.
-		+SSplitter::Slot()
-		[
-			SNew(SVerticalBox)
-			+SVerticalBox::Slot()
-			.AutoHeight()
-			[
-				SNew(SBorder)
-				.BorderImage(FAppStyle::GetBrush("ToolPanel.GroupBorder"))
-				.Padding(FMargin(3, 5))
-				[
-					SNew(SHorizontalBox)
-					+SHorizontalBox::Slot()
-					.FillWidth(1.f)
-					.VAlign(VAlign_Center)
-					.Padding(5.f)
-					[
-						SNew(SRichTextBlock)
-						.Text(LOCTEXT("BlendableParametersPanelTitle", "Blendable Parameters"))
-						.TransformPolicy(ETextTransformPolicy::ToUpper)
-						.DecoratorStyleSet(&FAppStyle::Get())
-						.TextStyle(FAppStyle::Get(), "DetailsView.CategoryTextStyle")
-					]
-					+SHorizontalBox::Slot()
-					.AutoWidth()
-					.VAlign(VAlign_Center)
-					[
-						SNew(SButton)
-						.ButtonStyle(FAppStyle::Get(), "SimpleButton")
-						.ContentPadding(FMargin(1, 0))
-						.ToolTipText(LOCTEXT("AddBlendableToolTip", "Add a blendable parameter"))
-						.OnClicked(this, &SCameraObjectInterfaceParametersPanel::OnAddBlendableParameter)
-						[
-							SNew(SImage)
-							.Image(FAppStyle::Get().GetBrush("Icons.PlusCircle"))
-							.ColorAndOpacity(FSlateColor::UseForeground())
-						]
-					]
-				]
-			]
-			+SVerticalBox::Slot()
-			.FillHeight(1.f)
-			[
-				SAssignNew(BlendableParametersListView, SListView<TObjectPtr<UCameraObjectInterfaceBlendableParameter>>)
-				.ListItemsSource(&CameraObject->Interface.BlendableParameters)
-				.OnGenerateRow(this, &SCameraObjectInterfaceParametersPanel::OnGenerateBlendableParameterRow)
-				.OnSelectionChanged(this, &SCameraObjectInterfaceParametersPanel::OnBlendableSelectionChanged)
-				.OnContextMenuOpening(this, &SCameraObjectInterfaceParametersPanel::OnBlendableParameterContextMenuOpening)
-				.HeaderRow
-				(
-					SNew(SHeaderRow)
-
-					+SHeaderRow::Column(ParameterTypeColumn)
-					.FillWidth(0.3f)
-					.DefaultLabel(LOCTEXT("ParameterTypeColumnLabel", "Type"))
-
-					+SHeaderRow::Column(ParameterNameColumn)
-					.FillWidth(0.5f)
-					.DefaultLabel(LOCTEXT("ParameterNameColumnLabel", "Name"))
-
-					+SHeaderRow::Column(ParameterIsPreBlendedColumn)
-					.ManualWidth(60)
-					.DefaultLabel(LOCTEXT("ParameterIsPreBlendedColumnLabel", "Pre-Blend"))
-
-					+SHeaderRow::Column(ParameterMessageColumn)
-					.FillWidth(0.2f)
-					.DefaultLabel(LOCTEXT("ParameterMessageColumnLabel", "Note"))
-				)
-			]
-		]
-
-		// The data parameters panel.
-		+SSplitter::Slot()
-		[
-			SNew(SVerticalBox)
-			+SVerticalBox::Slot()
-			.AutoHeight()
-			[
-				SNew(SBorder)
-				.BorderImage(FAppStyle::GetBrush("ToolPanel.GroupBorder"))
-				.Padding(FMargin(3, 5))
-				[
-					SNew(SHorizontalBox)
-					+SHorizontalBox::Slot()
-					.FillWidth(1.f)
-					.VAlign(VAlign_Center)
-					.Padding(5.f)
-					[
-						SNew(SRichTextBlock)
-						.Text(LOCTEXT("DataParametersPanelTitle", "Data Parameters"))
-						.TransformPolicy(ETextTransformPolicy::ToUpper)
-						.DecoratorStyleSet(&FAppStyle::Get())
-						.TextStyle(FAppStyle::Get(), "DetailsView.CategoryTextStyle")
-					]
-					+SHorizontalBox::Slot()
-					.AutoWidth()
-					.VAlign(VAlign_Center)
-					[
-						SNew(SButton)
-						.ButtonStyle(FAppStyle::Get(), "SimpleButton")
-						.ContentPadding(FMargin(1, 0))
-						.ToolTipText(LOCTEXT("AddDataParameterToolTip", "Add a data parameter"))
-						.OnClicked(this, &SCameraObjectInterfaceParametersPanel::OnAddDataParameter)
-						[
-							SNew(SImage)
-							.Image(FAppStyle::Get().GetBrush("Icons.PlusCircle"))
-							.ColorAndOpacity(FSlateColor::UseForeground())
-						]
-					]
-				]
-			]
-			+SVerticalBox::Slot()
-			.FillHeight(1.f)
-			[
-				SAssignNew(DataParametersListView, SListView<TObjectPtr<UCameraObjectInterfaceDataParameter>>)
-				.ListItemsSource(&CameraObject->Interface.DataParameters)
-				.OnGenerateRow(this, &SCameraObjectInterfaceParametersPanel::OnGenerateDataParameterRow)
-				.OnSelectionChanged(this, &SCameraObjectInterfaceParametersPanel::OnDataParameterSelectionChanged)
-				.OnContextMenuOpening(this, &SCameraObjectInterfaceParametersPanel::OnDataParameterContextMenuOpening)
-				.HeaderRow
-				(
-					SNew(SHeaderRow)
-
-					+SHeaderRow::Column(ParameterTypeColumn)
-					.FillWidth(0.3f)
-					.DefaultLabel(LOCTEXT("ParameterTypeColumnLabel", "Type"))
-
-					+SHeaderRow::Column(ParameterNameColumn)
-					.FillWidth(0.5f)
-					.DefaultLabel(LOCTEXT("ParameterNameColumnLabel", "Name"))
-
-					+SHeaderRow::Column(ParameterMessageColumn)
-					.FillWidth(0.2f)
-					.DefaultLabel(LOCTEXT("ParameterMessageColumnLabel", "Note"))
-				)
-			]
-		]
-	];
-}
-
-void SCameraObjectInterfaceParametersPanel::RequestListRefresh()
-{
-	bListRefreshRequested = true;
-}
-
-void SCameraObjectInterfaceParametersPanel::Tick(const FGeometry& AllottedGeometry, const double InCurrentTime, const float InDeltaTime)
-{
-	SCompoundWidget::Tick(AllottedGeometry, InCurrentTime, InDeltaTime);
-
-	if (bListRefreshRequested)
-	{
-		bListRefreshRequested = false;
-
-		BlendableParametersListView->RequestListRefresh();
-		DataParametersListView->RequestListRefresh();
-	}
-}
-
-TSharedRef<ITableRow> SCameraObjectInterfaceParametersPanel::OnGenerateBlendableParameterRow(TObjectPtr<UCameraObjectInterfaceBlendableParameter> Item, const TSharedRef<STableViewBase>& OwnerTable)
-{
-	return SNew(SCameraObjectInterfaceBlendableParameterTableRow, OwnerTable)
-		.Item(Item);
-}
-
-TSharedRef<ITableRow> SCameraObjectInterfaceParametersPanel::OnGenerateDataParameterRow(TObjectPtr<UCameraObjectInterfaceDataParameter> Item, const TSharedRef<STableViewBase>& OwnerTable)
-{
-	return SNew(SCameraObjectInterfaceDataParameterTableRow, OwnerTable)
-		.Item(Item);
-}
-
-void SCameraObjectInterfaceParametersPanel::OnBlendableSelectionChanged(TObjectPtr<UCameraObjectInterfaceBlendableParameter> Item, ESelectInfo::Type Type)
-{
-	Toolkit->OnInterfaceParameterSelected().Broadcast(Item);
-}
-
-void SCameraObjectInterfaceParametersPanel::OnDataParameterSelectionChanged(TObjectPtr<UCameraObjectInterfaceDataParameter> Item, ESelectInfo::Type Type)
-{
-	Toolkit->OnInterfaceParameterSelected().Broadcast(Item);
-}
-
-TSharedPtr<SWidget> SCameraObjectInterfaceParametersPanel::OnBlendableParameterContextMenuOpening()
-{
-	TArray<TObjectPtr<UCameraObjectInterfaceBlendableParameter>> SelectedItems = BlendableParametersListView->GetSelectedItems();
-	if (SelectedItems.Num() > 0)
-	{
-		return OnInterfaceParameterContextMenuOpening(BlendableParametersListView, SelectedItems[0]);
-	}
-	return nullptr;
-}
-
-TSharedPtr<SWidget> SCameraObjectInterfaceParametersPanel::OnDataParameterContextMenuOpening()
-{
-	TArray<TObjectPtr<UCameraObjectInterfaceDataParameter>> SelectedItems = DataParametersListView->GetSelectedItems();
-	if (SelectedItems.Num() > 0)
-	{
-		return OnInterfaceParameterContextMenuOpening(DataParametersListView, SelectedItems[0]);
-	}
-	return nullptr;
-}
-
-template<typename ItemType>
-TSharedPtr<SWidget> SCameraObjectInterfaceParametersPanel::OnInterfaceParameterContextMenuOpening(
-		TSharedPtr<SListView<TObjectPtr<ItemType>>> ListView, TObjectPtr<ItemType> Item)
-{
-	FMenuBuilder MenuBuilder(true, nullptr);
-
-	MenuBuilder.AddMenuEntry(
-			LOCTEXT("RenameInterfaceParameter", "Rename"),
-			LOCTEXT("RenameInterfaceParameterToolTip", "Renames this interface parameter"),
-			FSlateIcon(),
-			FExecuteAction::CreateSP(this, &SCameraObjectInterfaceParametersPanel::OnRenameInterfaceParameter, ListView, Item));
-	MenuBuilder.AddMenuEntry(
-			LOCTEXT("DeleteInterfaceParameter", "Delete"),
-			LOCTEXT("DeleteInterfaceParameterToolTip", "Deletes this interface parameter"),
-			FSlateIcon(),
-			FExecuteAction::CreateSP(this, &SCameraObjectInterfaceParametersPanel::OnDeleteInterfaceParameter, ListView, Item));
-
-	return MenuBuilder.MakeWidget();
-}
-
-template<typename ItemType>
-void SCameraObjectInterfaceParametersPanel::OnRenameInterfaceParameter(
-		TSharedPtr<SListView<TObjectPtr<ItemType>>> ListView, TObjectPtr<ItemType> Item)
-{
-	TSharedPtr<ITableRow> RowWidget = ListView->WidgetFromItem(Item);
-	if (!ensure(RowWidget))
-	{
-		return;
-	}
-
-	TSharedPtr<SCameraObjectInterfaceParameterTableRowBase<ItemType>> TypedRowWidget = 
-		StaticCastSharedPtr<SCameraObjectInterfaceParameterTableRowBase<ItemType>>(RowWidget);
-	TypedRowWidget->EnterNameEditingMode();
-}
-
-template<typename ItemType>
-void SCameraObjectInterfaceParametersPanel::OnDeleteInterfaceParameter(
-		TSharedPtr<SListView<TObjectPtr<ItemType>>> ListView, TObjectPtr<ItemType> Item)
-{
-	const FScopedTransaction Transaction(LOCTEXT("RemoveInterfaceParameter", "Remove Interface Parameter"));
-
-	CameraObject->Modify();
-
-	if constexpr (std::is_same_v<ItemType, UCameraObjectInterfaceBlendableParameter>)
-	{
-		const int32 NumRemoved = CameraObject->Interface.BlendableParameters.Remove(Item);
-		ensure(NumRemoved == 1);
-	}
-	else if constexpr (std::is_same_v<ItemType, UCameraObjectInterfaceDataParameter>)
-	{
-		const int32 NumRemoved = CameraObject->Interface.DataParameters.Remove(Item);
-		ensure(NumRemoved == 1);
-	}
-
-	CameraObject->EventHandlers.Notify(&ICameraObjectEventHandler::OnCameraObjectInterfaceChanged);
-
-	ListView->RequestListRefresh();
-}
-
-FReply SCameraObjectInterfaceParametersPanel::OnAddBlendableParameter()
-{
-	const FScopedTransaction Transaction(LOCTEXT("AddBlendableParameter", "Add Blendable Parameter"));
-
-	UCameraObjectInterfaceBlendableParameter* NewBlendableParameter = NewObject<UCameraObjectInterfaceBlendableParameter>(CameraObject, NAME_None, RF_Transactional);
-	NewBlendableParameter->InterfaceParameterName = NewBlendableParameter->GetName();
-
-	CameraObject->Modify();
-	CameraObject->Interface.BlendableParameters.Add(NewBlendableParameter);
-	CameraObject->EventHandlers.Notify(&ICameraObjectEventHandler::OnCameraObjectInterfaceChanged);
-
-	return FReply::Handled();
-}
-
-FReply SCameraObjectInterfaceParametersPanel::OnAddDataParameter()
-{
-	const FScopedTransaction Transaction(LOCTEXT("AddDataParameter", "Add Data Parameter"));
-
-	UCameraObjectInterfaceDataParameter* NewDataParameter = NewObject<UCameraObjectInterfaceDataParameter>(CameraObject, NAME_None, RF_Transactional);
-	NewDataParameter->InterfaceParameterName = NewDataParameter->GetName();
-
-	CameraObject->Modify();
-	CameraObject->Interface.DataParameters.Add(NewDataParameter);
-	CameraObject->EventHandlers.Notify(&ICameraObjectEventHandler::OnCameraObjectInterfaceChanged);
-
-	return FReply::Handled();
-}
-
-void SCameraObjectInterfaceParametersPanel::OnCameraObjectInterfaceChanged()
-{
-	bListRefreshRequested = true;
-}
-
 FCameraObjectInterfaceParametersToolkit::FCameraObjectInterfaceParametersToolkit()
 {
 	SAssignNew(PanelContainer, SBox);
diff --git a/Engine/Plugins/Cameras/GameplayCameras/Source/GameplayCamerasEditor/Private/Toolkits/CameraRigAssetEditorToolkit.cpp b/Engine/Plugins/Cameras/GameplayCameras/Source/GameplayCamerasEditor/Private/Toolkits/CameraRigAssetEditorToolkit.cpp
--- a/Engine/Plugins/Cameras/GameplayCameras/Source/GameplayCamerasEditor/Private/Toolkits/CameraRigAssetEditorToolkit.cpp
+++ b/Engine/Plugins/Cameras/GameplayCameras/Source/GameplayCamerasEditor/Private/Toolkits/CameraRigAssetEditorToolkit.cpp
@@ -223,8 +223,9 @@ void FCameraRigAssetEditorToolkit::CreateWidgets()
 	// Create the message log.
 	BuildLogToolkit->Initialize("CameraRigAssetBuildMessages");
 
-	// Hook-up the selection of interface parameters.
+	// Hook-up interface parameters events.
 	InterfaceParametersToolkit->OnInterfaceParameterSelected().AddSP(this, &FCameraRigAssetEditorToolkit::OnCameraObjectInterfaceParameterSelected);
+	InterfaceParametersToolkit->OnSearchInterfaceParameterNodes().AddSP(this, &FCameraRigAssetEditorToolkit::OnSearchCameraObjectInterfaceParameterNodes);
 }
 
 void FCameraRigAssetEditorToolkit::RegisterToolbar()
@@ -333,8 +334,14 @@ void FCameraRigAssetEditorToolkit::PostRegenerateMenusAndToolbars()
 
 void FCameraRigAssetEditorToolkit::OnCameraObjectInterfaceParameterSelected(UCameraObjectInterfaceParameterBase* Object)
 {
-	TSharedPtr<SCameraRigAssetEditor> CameraRigEditor = Impl->GetCameraRigAssetEditor();
-	CameraRigEditor->FindAndJumpToObjectNode(Object);
+	DetailsView->SetObject(Object);
+}
+
+void FCameraRigAssetEditorToolkit::OnSearchCameraObjectInterfaceParameterNodes(UCameraObjectInterfaceParameterBase* Object)
+{
+	TabManager->TryInvokeTab(SearchTabId);
+
+	SearchWidget->Search(Object->InterfaceParameterName);
 }
 
 void FCameraRigAssetEditorToolkit::OnBuild()
@@ -356,6 +363,8 @@ void FCameraRigAssetEditorToolkit::OnBuild()
 		TabManager->TryInvokeTab(MessagesTabId);
 	}
 
+	DetailsView->RequestForceRefresh();
+
 	FCameraRigPackages BuiltPackages;
 	CameraRigAsset->GatherPackages(BuiltPackages);
 
diff --git a/Engine/Plugins/Cameras/GameplayCameras/Source/GameplayCamerasEditor/Private/Toolkits/CameraShakeAssetEditorToolkit.cpp b/Engine/Plugins/Cameras/GameplayCameras/Source/GameplayCamerasEditor/Private/Toolkits/CameraShakeAssetEditorToolkit.cpp
--- a/Engine/Plugins/Cameras/GameplayCameras/Source/GameplayCamerasEditor/Private/Toolkits/CameraShakeAssetEditorToolkit.cpp
+++ b/Engine/Plugins/Cameras/GameplayCameras/Source/GameplayCamerasEditor/Private/Toolkits/CameraShakeAssetEditorToolkit.cpp
@@ -394,6 +394,8 @@ void FCameraShakeAssetEditorToolkit::OnBuild()
 	{
 		TabManager->TryInvokeTab(MessagesTabId);
 	}
+
+	DetailsView->RequestForceRefresh();
 }
 
 void FCameraShakeAssetEditorToolkit::OnFindInCameraShake()
diff --git a/Engine/Plugins/Cameras/GameplayCameras/Source/GameplayCamerasEditor/Public/Customizations/CameraObjectInterfaceParameterDetailsCustomization.h b/Engine/Plugins/Cameras/GameplayCameras/Source/GameplayCamerasEditor/Public/Customizations/CameraObjectInterfaceParameterDetailsCustomization.h
new file mode 100644
--- /dev/null
+++ b/Engine/Plugins/Cameras/GameplayCameras/Source/GameplayCamerasEditor/Public/Customizations/CameraObjectInterfaceParameterDetailsCustomization.h
@@ -0,0 +1,31 @@
+// Copyright Epic Games, Inc. All Rights Reserved.
+
+#pragma once
+
+#include "IDetailCustomization.h"
+#include "Templates/SharedPointerFwd.h"
+#include "UObject/WeakObjectPtr.h"
+
+class IPropertyUtilities;
+class UBaseCameraObject;
+
+namespace UE::Cameras
+{
+
+/**
+ * Details customization for a camera object parameter.
+ */
+class FCameraObjectInterfaceParameterDetailsCustomization : public IDetailCustomization
+{
+public:
+
+	static TSharedRef<IDetailCustomization> MakeInstance();
+
+public:
+
+	// IDetailCustomization interface.
+	virtual void CustomizeDetails(IDetailLayoutBuilder& DetailBuilder) override;
+};
+
+}  // namespace UE::Cameras
+
diff --git a/Engine/Plugins/Cameras/GameplayCameras/Source/GameplayCamerasEditor/Public/Customizations/CameraParameterDetailsCustomizations.h b/Engine/Plugins/Cameras/GameplayCameras/Source/GameplayCamerasEditor/Public/Customizations/CameraParameterDetailsCustomizations.h
--- a/Engine/Plugins/Cameras/GameplayCameras/Source/GameplayCamerasEditor/Public/Customizations/CameraParameterDetailsCustomizations.h
+++ b/Engine/Plugins/Cameras/GameplayCameras/Source/GameplayCamerasEditor/Public/Customizations/CameraParameterDetailsCustomizations.h
@@ -45,6 +45,7 @@ public:
 
 protected:
 
+	virtual bool HasOverride(void* InRawData) = 0;
 	virtual bool HasNonUserOverride(void* InRawData) = 0;
 	virtual void SetParameterVariable(void* InRawData, UCameraVariableAsset* InVariable) = 0;
 
@@ -63,6 +64,7 @@ private:
 		UCameraVariableAsset* CommonVariable = nullptr;
 		ECameraVariableValue VariableValue = ECameraVariableValue::NotSet;
 		bool bHasNonUserOverride = false;
+		bool bHasOverride = false;
 
 		FText InfoText;
 		FText ErrorText;
@@ -116,6 +118,7 @@ protected:
 class F##ValueName##CameraParameterDetailsCustomization : public FCameraParameterDetailsCustomization\
 {\
 protected:\
+	virtual bool HasOverride(void* InRawData) override;\
 	virtual bool HasNonUserOverride(void* InRawData) override;\
 	virtual void SetParameterVariable(void* InRawData, UCameraVariableAsset* InVariable) override;\
 };
diff --git a/Engine/Plugins/Cameras/GameplayCameras/Source/GameplayCamerasEditor/Public/Editors/CameraNodeGraphSchema.h b/Engine/Plugins/Cameras/GameplayCameras/Source/GameplayCamerasEditor/Public/Editors/CameraNodeGraphSchema.h
deleted file mode 100644
diff --git a/Engine/Plugins/Cameras/GameplayCameras/Source/GameplayCamerasEditor/Public/Editors/CameraObjectGraphSchemaBase.h b/Engine/Plugins/Cameras/GameplayCameras/Source/GameplayCamerasEditor/Public/Editors/CameraObjectGraphSchemaBase.h
new file mode 100644
--- /dev/null
+++ b/Engine/Plugins/Cameras/GameplayCameras/Source/GameplayCamerasEditor/Public/Editors/CameraObjectGraphSchemaBase.h
@@ -0,0 +1,115 @@
+// Copyright Epic Games, Inc. All Rights Reserved.
+
+#pragma once
+
+#include "Compat/EditorCompat.h"
+#include "Core/CameraObjectInterface.h"
+#include "Core/CameraObjectInterfaceParameterDefinition.h"
+#include "EdGraph/EdGraphPin.h"
+#include "Editors/CameraNodeGraphPinColors.h"
+#include "Editors/ObjectTreeGraphSchema.h"
+
+#include "CameraObjectGraphSchemaBase.generated.h"
+
+class UCameraObjectInterfaceParameterGraphNode;
+struct FObjectTreeGraphConfig;
+
+/**
+ * Base schema class for any kind of camera node graph.
+ */
+UCLASS()
+class UCameraObjectGraphSchemaBase : public UObjectTreeGraphSchema
+{
+	GENERATED_BODY()
+
+public:
+
+	static const FName PC_CameraParameter;			// A camera parameter pin.
+	static const FName PC_CameraVariableReference;	// A variable reference pin.
+	static const FName PC_CameraContextData;		// A context data pin.
+
+	UCameraObjectGraphSchemaBase(const FObjectInitializer& ObjInit);
+
+	/** Builds the config for a graph managed by this schema. */
+	FObjectTreeGraphConfig BuildGraphConfig() const;
+
+protected:
+
+	// UEdGraphSchema interface.
+	virtual void GetGraphContextActions(FGraphContextMenuBuilder& ContextMenuBuilder) const override;
+	virtual const FPinConnectionResponse CanCreateConnection(const UEdGraphPin* A, const UEdGraphPin* B) const override;
+	virtual bool OnBreakCustomPinLinks(UEdGraphPin& TargetPin) const;
+	virtual bool OnBreakSingleCustomPinLink(UEdGraphPin* SourcePin, UEdGraphPin* TargetPin) const;
+	virtual FLinearColor GetPinTypeColor(const FEdGraphPinType& PinType) const override;
+
+	// UObjectTreeGraphSchema interface.
+	virtual void OnCreateAllNodes(UObjectTreeGraph* InGraph, const FCreatedNodes& InCreatedNodes) const override;
+	virtual bool OnTryCreateCustomConnection(UEdGraphPin* A, UEdGraphPin* B) const;
+
+protected:
+
+	// UCameraObjectGraphSchemaBase interface.
+	virtual void OnBuildGraphConfig(FObjectTreeGraphConfig& InOutGraphConfig) const {}
+
+private:
+
+	void CreateValueFlowConnections(UObjectTreeGraph* InGraph, const FCreatedNodes& InCreatedNodes) const;
+	UEdGraphPin* FindPinByType(UEdGraphNode* InNode, const FName& InPinCategory) const;
+
+	UE::Cameras::FCameraNodeGraphPinColors PinColors;
+};
+
+/**
+ * Graph editor action for adding a new camera rig parameter, plus an associated getter node in the graph.
+ */
+USTRUCT()
+struct FCameraObjectGraphSchemaAction_NewInterfaceParameterNode : public FEdGraphSchemaAction
+{
+	GENERATED_BODY()
+
+public:
+
+	/** The new parameter's definition. */
+	UPROPERTY()
+	FCameraObjectInterfaceParameterDefinition ParameterDefinition;
+
+public:
+
+	FCameraObjectGraphSchemaAction_NewInterfaceParameterNode();
+	FCameraObjectGraphSchemaAction_NewInterfaceParameterNode(FText InNodeCategory, FText InMenuDesc, FText InToolTip, const int32 InGrouping = 0, FText InKeywords = FText());
+
+public:
+
+	// FEdGraphSchemaAction interface.
+	static FName StaticGetTypeId() { static FName Type("FCameraObjectGraphSchemaAction_NewInterfaceParameterNode"); return Type; }
+	virtual FName GetTypeId() const override { return StaticGetTypeId(); }
+	virtual UEdGraphNode* PerformAction(UEdGraph* ParentGraph, UEdGraphPin* FromPin, FPerformGraphActionLocation Location, bool bSelectNewNode = true) override;
+};
+
+/**
+ * Graph editor action for adding a new getter node for an existing camera rig parameter.
+ */
+USTRUCT()
+struct FCameraObjectGraphSchemaAction_AddInterfaceParameterGetterNode : public FEdGraphSchemaAction
+{
+	GENERATED_BODY()
+	
+public:
+
+	/** The existing camera rig parameter to create a getter node for. */
+	UPROPERTY()
+	TObjectPtr<UCameraObjectInterfaceParameterBase> InterfaceParameter;
+
+public:
+	
+	FCameraObjectGraphSchemaAction_AddInterfaceParameterGetterNode();
+	FCameraObjectGraphSchemaAction_AddInterfaceParameterGetterNode(FText InNodeCategory, FText InMenuDesc, FText InToolTip, const int32 InGrouping = 0, FText InKeywords = FText());
+
+public:
+
+	// FEdGraphSchemaAction interface.
+	static FName StaticGetTypeId() { static FName Type("FCameraObjectGraphSchemaAction_AddInterfaceParameterGetterNode"); return Type; }
+	virtual FName GetTypeId() const override { return StaticGetTypeId(); }
+	virtual UEdGraphNode* PerformAction(UEdGraph* ParentGraph, UEdGraphPin* FromPin, FPerformGraphActionLocation Location, bool bSelectNewNode = true) override;
+};
+
diff --git a/Engine/Plugins/Cameras/GameplayCameras/Source/GameplayCamerasEditor/Public/Editors/CameraObjectInterfaceParameterGraphNode.h b/Engine/Plugins/Cameras/GameplayCameras/Source/GameplayCamerasEditor/Public/Editors/CameraObjectInterfaceParameterGraphNode.h
--- a/Engine/Plugins/Cameras/GameplayCameras/Source/GameplayCamerasEditor/Public/Editors/CameraObjectInterfaceParameterGraphNode.h
+++ b/Engine/Plugins/Cameras/GameplayCameras/Source/GameplayCamerasEditor/Public/Editors/CameraObjectInterfaceParameterGraphNode.h
@@ -9,7 +9,7 @@
 class UCameraObjectInterfaceParameterBase;
 
 /**
- * Custom graph editor node for a camera rig parameter.
+ * Custom graph editor node for a camera rig parameter getter.
  */
 UCLASS()
 class UCameraObjectInterfaceParameterGraphNode : public UObjectTreeGraphNode
@@ -21,12 +21,13 @@ public:
 	/** Creates a new graph node. */
 	UCameraObjectInterfaceParameterGraphNode(const FObjectInitializer& ObjInit);
 
-	/** Gets the underlying object as a camera rig interface parameter. */
+	/** Finds the camera interface parameter. */
 	UCameraObjectInterfaceParameterBase* GetInterfaceParameter() const;
 
 public:
 
 	// UEdGraphNode interface
+	virtual FText GetNodeTitle(ENodeTitleType::Type TitleType) const override;
 	virtual TSharedPtr<SGraphNode> CreateVisualWidget() override;
 };
 
diff --git a/Engine/Plugins/Cameras/GameplayCameras/Source/GameplayCamerasEditor/Public/Editors/CameraRigCameraNodeGraphSchema.h b/Engine/Plugins/Cameras/GameplayCameras/Source/GameplayCamerasEditor/Public/Editors/CameraRigCameraNodeGraphSchema.h
--- a/Engine/Plugins/Cameras/GameplayCameras/Source/GameplayCamerasEditor/Public/Editors/CameraRigCameraNodeGraphSchema.h
+++ b/Engine/Plugins/Cameras/GameplayCameras/Source/GameplayCamerasEditor/Public/Editors/CameraRigCameraNodeGraphSchema.h
@@ -2,7 +2,7 @@
 
 #pragma once
 
-#include "Editors/CameraNodeGraphSchema.h"
+#include "Editors/CameraObjectGraphSchemaBase.h"
 
 #include "CameraRigCameraNodeGraphSchema.generated.h"
 
@@ -10,7 +10,7 @@
  * Schema class for camera node graph.
  */
 UCLASS()
-class UCameraRigCameraNodeGraphSchema : public UCameraNodeGraphSchema
+class UCameraRigCameraNodeGraphSchema : public UCameraObjectGraphSchemaBase
 {
 	GENERATED_BODY()
 
@@ -18,11 +18,12 @@ public:
 
 	UCameraRigCameraNodeGraphSchema(const FObjectInitializer& ObjInit);
 
-	FObjectTreeGraphConfig BuildGraphConfig() const;
-
 protected:
 
 	// UObjectTreeGraphSchema interface.
 	virtual void CollectAllObjects(UObjectTreeGraph* InGraph, TSet<UObject*>& OutAllObjects) const override;
+
+	// UCameraObjectGraphSchemaBase interface.
+	virtual void OnBuildGraphConfig(FObjectTreeGraphConfig& InOutGraphConfig) const override;
 };
 
diff --git a/Engine/Plugins/Cameras/GameplayCameras/Source/GameplayCamerasEditor/Public/Editors/CameraRigTransitionGraphSchemaBase.h b/Engine/Plugins/Cameras/GameplayCameras/Source/GameplayCamerasEditor/Public/Editors/CameraRigTransitionGraphSchemaBase.h
--- a/Engine/Plugins/Cameras/GameplayCameras/Source/GameplayCamerasEditor/Public/Editors/CameraRigTransitionGraphSchemaBase.h
+++ b/Engine/Plugins/Cameras/GameplayCameras/Source/GameplayCamerasEditor/Public/Editors/CameraRigTransitionGraphSchemaBase.h
@@ -2,7 +2,7 @@
 
 #pragma once
 
-#include "Editors/ObjectTreeGraphSchema.h"
+#include "Editors/CameraObjectGraphSchemaBase.h"
 
 #include "CameraRigTransitionGraphSchemaBase.generated.h"
 
@@ -12,22 +12,16 @@ struct FObjectTreeGraphConfig;
  * Base schema class for camera transition graph.
  */
 UCLASS()
-class UCameraRigTransitionGraphSchemaBase : public UObjectTreeGraphSchema
+class UCameraRigTransitionGraphSchemaBase : public UCameraObjectGraphSchemaBase
 {
 	GENERATED_BODY()
 
-public:
-
-	FObjectTreeGraphConfig BuildGraphConfig() const;
-
 protected:
 
 	// UObjectTreeGraphSchema interface.
 	virtual void CollectAllObjects(UObjectTreeGraph* InGraph, TSet<UObject*>& OutAllObjects) const override;
 
-protected:
-
-	// UCameraRigTransitionGraphSchemaBase interface.
-	virtual void OnBuildGraphConfig(FObjectTreeGraphConfig& InOutGraphConfig) const {}
+	// UCameraObjectGraphSchemaBase interface.
+	virtual void OnBuildGraphConfig(FObjectTreeGraphConfig& InOutGraphConfig) const override;
 };
 
diff --git a/Engine/Plugins/Cameras/GameplayCameras/Source/GameplayCamerasEditor/Public/Editors/CameraShakeCameraNodeGraphSchema.h b/Engine/Plugins/Cameras/GameplayCameras/Source/GameplayCamerasEditor/Public/Editors/CameraShakeCameraNodeGraphSchema.h
--- a/Engine/Plugins/Cameras/GameplayCameras/Source/GameplayCamerasEditor/Public/Editors/CameraShakeCameraNodeGraphSchema.h
+++ b/Engine/Plugins/Cameras/GameplayCameras/Source/GameplayCamerasEditor/Public/Editors/CameraShakeCameraNodeGraphSchema.h
@@ -2,7 +2,7 @@
 
 #pragma once
 
-#include "Editors/CameraNodeGraphSchema.h"
+#include "Editors/CameraObjectGraphSchemaBase.h"
 
 #include "CameraShakeCameraNodeGraphSchema.generated.h"
 
@@ -10,7 +10,7 @@
  * Schema class for camera node graph.
  */
 UCLASS()
-class UCameraShakeCameraNodeGraphSchema : public UCameraNodeGraphSchema
+class UCameraShakeCameraNodeGraphSchema : public UCameraObjectGraphSchemaBase
 {
 	GENERATED_BODY()
 
@@ -18,6 +18,12 @@ public:
 
 	UCameraShakeCameraNodeGraphSchema(const FObjectInitializer& ObjInit);
 
-	FObjectTreeGraphConfig BuildGraphConfig() const;
+protected:
+
+	// UObjectTreeGraphSchema interface.
+	virtual void CollectAllObjects(UObjectTreeGraph* InGraph, TSet<UObject*>& OutAllObjects) const override;
+
+	// UCameraObjectGraphSchemaBase interface.
+	virtual void OnBuildGraphConfig(FObjectTreeGraphConfig& InOutGraphConfig) const override;
 };
 
diff --git a/Engine/Plugins/Cameras/GameplayCameras/Source/GameplayCamerasEditor/Public/Editors/SCameraObjectInterfaceParameterGraphNode.h b/Engine/Plugins/Cameras/GameplayCameras/Source/GameplayCamerasEditor/Public/Editors/SCameraObjectInterfaceParameterGraphNode.h
--- a/Engine/Plugins/Cameras/GameplayCameras/Source/GameplayCamerasEditor/Public/Editors/SCameraObjectInterfaceParameterGraphNode.h
+++ b/Engine/Plugins/Cameras/GameplayCameras/Source/GameplayCamerasEditor/Public/Editors/SCameraObjectInterfaceParameterGraphNode.h
@@ -7,7 +7,7 @@
 class UCameraObjectInterfaceParameterGraphNode;
 
 /**
- * Custom graph editor node widget for a camera rig parameter node.
+ * Custom graph editor node widget for a camera rig parameter getter node.
  */
 class SCameraObjectInterfaceParameterGraphNode : public SObjectTreeGraphNode
 {
diff --git a/Engine/Plugins/Cameras/GameplayCameras/Source/GameplayCamerasEditor/Public/Editors/SCameraObjectInterfaceParametersPanel.h b/Engine/Plugins/Cameras/GameplayCameras/Source/GameplayCamerasEditor/Public/Editors/SCameraObjectInterfaceParametersPanel.h
new file mode 100644
--- /dev/null
+++ b/Engine/Plugins/Cameras/GameplayCameras/Source/GameplayCamerasEditor/Public/Editors/SCameraObjectInterfaceParametersPanel.h
@@ -0,0 +1,106 @@
+// Copyright Epic Games, Inc. All Rights Reserved.
+
+#pragma once
+
+#include "Core/BaseCameraObject.h"
+#include "Core/CameraEventHandler.h"
+#include "EdGraphSchema_K2.h"
+#include "Templates/SharedPointerFwd.h"
+#include "Widgets/SCompoundWidget.h"
+
+#include "SCameraObjectInterfaceParametersPanel.generated.h"
+
+class ITableRow;
+class STableViewBase;
+template <typename ItemType> class SListView;
+
+namespace UE::Cameras
+{
+
+class FCameraObjectInterfaceParametersToolkit;
+
+/**
+ * The overall interface parameters panel, showing two sub-panels, one for blendable parameters, and
+ * one for data parameters.
+ */
+class SCameraObjectInterfaceParametersPanel 
+	: public SCompoundWidget
+	, public ICameraObjectEventHandler
+{
+public:
+
+	SLATE_BEGIN_ARGS(SCameraObjectInterfaceParametersPanel)
+	{}
+	SLATE_END_ARGS()
+
+	void Construct(const FArguments& Args, FCameraObjectInterfaceParametersToolkit* OwnerToolkit);
+
+	void RequestListRefresh();
+
+protected:
+
+	// SWidget interface.
+	virtual void Tick(const FGeometry& AllottedGeometry, const double InCurrentTime, const float InDeltaTime) override;
+
+	// ICameraObjectEventHandler interface.
+	virtual void OnCameraObjectInterfaceChanged() override;
+
+private:
+
+	TSharedRef<ITableRow> OnGenerateBlendableParameterRow(TObjectPtr<UCameraObjectInterfaceBlendableParameter> Item, const TSharedRef<STableViewBase>& OwnerTable);
+	TSharedRef<ITableRow> OnGenerateDataParameterRow(TObjectPtr<UCameraObjectInterfaceDataParameter> Item, const TSharedRef<STableViewBase>& OwnerTable);
+
+	void OnBlendableSelectionChanged(TObjectPtr<UCameraObjectInterfaceBlendableParameter> Item, ESelectInfo::Type Type);
+	void OnDataParameterSelectionChanged(TObjectPtr<UCameraObjectInterfaceDataParameter> Item, ESelectInfo::Type Type);
+
+	TSharedPtr<SWidget> OnBlendableParameterContextMenuOpening();
+	TSharedPtr<SWidget> OnDataParameterContextMenuOpening();
+
+	void OnSearchInterfaceParameter(TObjectPtr<UCameraObjectInterfaceParameterBase> Item);
+
+	template<typename ItemType>
+	TSharedPtr<SWidget> OnInterfaceParameterContextMenuOpening(
+			TSharedPtr<SListView<TObjectPtr<ItemType>>> ListView, TObjectPtr<ItemType> Item);
+	template<typename ItemType>
+	void OnRenameInterfaceParameter(
+			TSharedPtr<SListView<TObjectPtr<ItemType>>> ListView, TObjectPtr<ItemType> Item);
+	template<typename ItemType>
+	void OnDeleteInterfaceParameter(
+			TSharedPtr<SListView<TObjectPtr<ItemType>>> ListView, TObjectPtr<ItemType> Item);
+
+	FReply OnAddBlendableParameter();
+	FReply OnAddDataParameter();
+
+private:
+
+	UBaseCameraObject* CameraObject = nullptr;
+	FCameraObjectInterfaceParametersToolkit* Toolkit = nullptr;
+
+	TCameraEventHandler<ICameraObjectEventHandler> EventHandler;
+
+	TSharedPtr<SListView<TObjectPtr<UCameraObjectInterfaceBlendableParameter>>> BlendableParametersListView;
+	TSharedPtr<SListView<TObjectPtr<UCameraObjectInterfaceDataParameter>>> DataParametersListView;
+
+	bool bListRefreshRequested = false;
+};
+
+}  // namespace UE::Cameras
+
+// Dummy Blueprint schema used only for getting proper icons and colors for parameter types.
+UCLASS(MinimalAPI, Hidden)
+class UEdGraphSchema_CameraNodeK2 : public UEdGraphSchema_K2
+{
+	GENERATED_BODY()
+
+public:
+
+	virtual bool SupportsPinTypeContainer(TWeakPtr<const FEdGraphSchemaAction> SchemaAction, const FEdGraphPinType& PinType, const EPinContainerType& ContainerType) const
+	{
+		if (ContainerType == EPinContainerType::None || ContainerType == EPinContainerType::Array)
+		{
+			return Super::SupportsPinTypeContainer(SchemaAction, PinType, ContainerType);
+		}
+		return false;
+	}
+};
+
diff --git a/Engine/Plugins/Cameras/GameplayCameras/Source/GameplayCamerasEditor/Public/Editors/SCameraRigAssetEditor.h b/Engine/Plugins/Cameras/GameplayCameras/Source/GameplayCamerasEditor/Public/Editors/SCameraRigAssetEditor.h
--- a/Engine/Plugins/Cameras/GameplayCameras/Source/GameplayCamerasEditor/Public/Editors/SCameraRigAssetEditor.h
+++ b/Engine/Plugins/Cameras/GameplayCameras/Source/GameplayCamerasEditor/Public/Editors/SCameraRigAssetEditor.h
@@ -3,7 +3,6 @@
 #pragma once
 
 #include "Compat/EditorCompat.h"
-#include "Core/CameraRigAsset.h"
 #include "CoreTypes.h"
 #include "GraphEditor.h"
 #include "Templates/SharedPointerFwd.h"
@@ -14,6 +13,7 @@ class FAssetEditorToolkit;
 class IDetailsView;
 class SBox;
 class SObjectTreeGraphEditor;
+class UCameraRigAsset;
 class UEdGraphNode;
 class UObjectTreeGraph;
 struct FObjectTreeGraphConfig;
@@ -38,7 +38,6 @@ enum class ECameraRigAssetEditorMode
  */
 class SCameraRigAssetEditor 
 	: public SCompoundWidget
-	, public ICameraObjectEventHandler
 {
 public:
 
@@ -88,11 +87,6 @@ public:
 	/** Removes a previous added callback. */
 	void RemoveOnAnyGraphChanged(FDelegateUserObjectConst InUserObject);
 
-protected:
-
-	// ICameraObjectEventHandler interface.
-	virtual void OnCameraObjectInterfaceChanged() override;
-
 protected:
 
 	void CreateGraphEditors();
@@ -114,9 +108,6 @@ private:
 	/** The asset being edited */
 	TObjectPtr<UCameraRigAsset> CameraRigAsset;
 
-	/** Event handler */
-	TCameraEventHandler<ICameraObjectEventHandler> EventHandler;
-
 	/** Reference to the details view */
 	TSharedPtr<IDetailsView> DetailsView;
 
diff --git a/Engine/Plugins/Cameras/GameplayCameras/Source/GameplayCamerasEditor/Public/Toolkits/CameraObjectInterfaceParametersToolkit.h b/Engine/Plugins/Cameras/GameplayCameras/Source/GameplayCamerasEditor/Public/Toolkits/CameraObjectInterfaceParametersToolkit.h
--- a/Engine/Plugins/Cameras/GameplayCameras/Source/GameplayCamerasEditor/Public/Toolkits/CameraObjectInterfaceParametersToolkit.h
+++ b/Engine/Plugins/Cameras/GameplayCameras/Source/GameplayCamerasEditor/Public/Toolkits/CameraObjectInterfaceParametersToolkit.h
@@ -7,8 +7,6 @@
 #include "Templates/SharedPointer.h"
 #include "UObject/ObjectPtr.h"
 
-#include "CameraObjectInterfaceParametersToolkit.generated.h"
-
 class SBox;
 class SWidget;
 class UBaseCameraObject;
@@ -41,9 +39,14 @@ public:
 	/** Gets the panel widget. */
 	TSharedPtr<SWidget> GetInterfaceParametersPanel() const;
 
+public:
+
 	/** Delegate invoked when a parmeter is selected in the panel. */
 	FOnCameraObjectInterfaceParameterEvent& OnInterfaceParameterSelected() { return OnInterfaceParameterSelectedDelegate; }
 
+	/** Delegate invoked when the user requests finding getter nodes in the graph. */
+	FOnCameraObjectInterfaceParameterEvent& OnSearchInterfaceParameterNodes() { return OnSearchInterfaceParameterNodesDelegate; }
+
 protected:
 
 	// FEditorUndoClient interface
@@ -55,6 +58,7 @@ private:
 	TObjectPtr<UBaseCameraObject> CameraObject;
 
 	FOnCameraObjectInterfaceParameterEvent OnInterfaceParameterSelectedDelegate;
+	FOnCameraObjectInterfaceParameterEvent OnSearchInterfaceParameterNodesDelegate;
 
 	TSharedPtr<SBox> PanelContainer;
 	TSharedPtr<SCameraObjectInterfaceParametersPanel> Panel;
@@ -62,20 +66,3 @@ private:
 
 }  // namespace UE::Cameras
 
-UCLASS(MinimalAPI, Hidden)
-class UEdGraphSchema_CameraNodeK2 : public UEdGraphSchema_K2
-{
-	GENERATED_BODY()
-
-public:
-
-	virtual bool SupportsPinTypeContainer(TWeakPtr<const FEdGraphSchemaAction> SchemaAction, const FEdGraphPinType& PinType, const EPinContainerType& ContainerType) const
-	{
-		if (ContainerType == EPinContainerType::None || ContainerType == EPinContainerType::Array)
-		{
-			return Super::SupportsPinTypeContainer(SchemaAction, PinType, ContainerType);
-		}
-		return false;
-	}
-};
-
diff --git a/Engine/Plugins/Cameras/GameplayCameras/Source/GameplayCamerasEditor/Public/Toolkits/CameraRigAssetEditorToolkit.h b/Engine/Plugins/Cameras/GameplayCameras/Source/GameplayCamerasEditor/Public/Toolkits/CameraRigAssetEditorToolkit.h
--- a/Engine/Plugins/Cameras/GameplayCameras/Source/GameplayCamerasEditor/Public/Toolkits/CameraRigAssetEditorToolkit.h
+++ b/Engine/Plugins/Cameras/GameplayCameras/Source/GameplayCamerasEditor/Public/Toolkits/CameraRigAssetEditorToolkit.h
@@ -66,6 +66,7 @@ private:
 	void OnCurvesTabClosed(TSharedRef<SDockTab> InTab);
 
 	void OnCameraObjectInterfaceParameterSelected(UCameraObjectInterfaceParameterBase* Object);
+	void OnSearchCameraObjectInterfaceParameterNodes(UCameraObjectInterfaceParameterBase* Object);
 
 	void OnBuild();
 	void OnFindInCameraRig();
