Cameras: expose the blend type as a data pin on the easing blend node

Also fix a few problems with the Context Data Reader utility struct.

#Synced-CL 48619321

diff --git a/Engine/Plugins/Cameras/GameplayCameras/Source/GameplayCameras/Private/Nodes/EasingBlendCameraNode.cpp b/Engine/Plugins/Cameras/GameplayCameras/Source/GameplayCameras/Private/Nodes/EasingBlendCameraNode.cpp
--- a/Engine/Plugins/Cameras/GameplayCameras/Source/GameplayCameras/Private/Nodes/EasingBlendCameraNode.cpp
+++ b/Engine/Plugins/Cameras/GameplayCameras/Source/GameplayCameras/Private/Nodes/EasingBlendCameraNode.cpp
@@ -2,6 +2,7 @@
 
 #include "Nodes/Blends/EasingBlendCameraNode.h"
 
+#include "Core/CameraContextDataReader.h"
 #include "Math/UnrealMathUtility.h"
 
 #include UE_INLINE_GENERATED_CPP_BY_NAME(EasingBlendCameraNode)
@@ -127,18 +128,33 @@ class FEasingBlendCameraNodeEvaluator : public FSimpleFixedTimeBlendCameraNodeEv
 	UE_DECLARE_BLEND_CAMERA_NODE_EVALUATOR_EX(GAMEPLAYCAMERAS_API, FEasingBlendCameraNodeEvaluator, FSimpleFixedTimeBlendCameraNodeEvaluator)
 
 protected:
+
+	virtual void OnInitialize(const FCameraNodeEvaluatorInitializeParams& Params, FCameraNodeEvaluationResult& OutResult) override;
 	virtual void OnComputeBlendFactor(const FCameraNodeEvaluationParams& Params, FSimpleBlendCameraNodeEvaluationResult& OutResult) override;
+
+private:
+
+	EEasingCameraBlendType BlendType;
 };
 
 UE_DEFINE_BLEND_CAMERA_NODE_EVALUATOR(FEasingBlendCameraNodeEvaluator)
 
+void FEasingBlendCameraNodeEvaluator::OnInitialize(const FCameraNodeEvaluatorInitializeParams& Params, FCameraNodeEvaluationResult& OutResult)
+{
+	Super::OnInitialize(Params, OutResult);
+
+	const UEasingBlendCameraNode* BlendNode = GetCameraNodeAs<UEasingBlendCameraNode>();
+	TCameraContextDataReader<EEasingCameraBlendType> BlendTypeReader;
+	BlendTypeReader.Initialize(&BlendNode->BlendType, BlendNode->BlendTypeDataID);
+	BlendType = BlendTypeReader.Get(OutResult.ContextDataTable);
+}
+
 void FEasingBlendCameraNodeEvaluator::OnComputeBlendFactor(const FCameraNodeEvaluationParams& Params, FSimpleBlendCameraNodeEvaluationResult& OutResult)
 {
 	using namespace UE::Cameras::Internal;
 
-	const UEasingBlendCameraNode* BlendNode = GetCameraNodeAs<UEasingBlendCameraNode>();
 	const float TimeFactor = GetTimeFactor();
-	OutResult.BlendFactor = EvaluateEasing(BlendNode->BlendType, TimeFactor);
+	OutResult.BlendFactor = EvaluateEasing(BlendType, TimeFactor);
 }
 
 }  // namespace UE::Cameras
diff --git a/Engine/Plugins/Cameras/GameplayCameras/Source/GameplayCameras/Public/Core/CameraContextDataReader.h b/Engine/Plugins/Cameras/GameplayCameras/Source/GameplayCameras/Public/Core/CameraContextDataReader.h
--- a/Engine/Plugins/Cameras/GameplayCameras/Source/GameplayCameras/Public/Core/CameraContextDataReader.h
+++ b/Engine/Plugins/Cameras/GameplayCameras/Source/GameplayCameras/Public/Core/CameraContextDataReader.h
@@ -19,6 +19,11 @@ struct TCameraContextDataReader
 		Initialize(Parameter);
 	}
 
+	TCameraContextDataReader(const DataType* InDefaultValuePtr, FCameraContextDataID InDataID)
+	{
+		Initialize(InDefaultValuePtr, InDataID);
+	}
+
 	template<typename ParameterType>
 	void Initialize(const ParameterType& Parameter)
 	{
@@ -32,11 +37,19 @@ struct TCameraContextDataReader
 		ensureMsgf(DefaultValuePtr, TEXT("The given parameter doesn't have a value!"));
 	}
 
+	void Initialize(const DataType* InDefaultValuePtr, FCameraContextDataID InDataID)
+	{
+		DefaultValuePtr = InDefaultValuePtr;
+		DataID = InDataID;
+
+		ensureMsgf(DefaultValuePtr, TEXT("No default value was provided!"));
+	}
+
 	DataType Get(const FCameraContextDataTable& ContextDataTable) const
 	{
 		if (!DataID.IsValid())
 		{
-			return *DefaultValuePtr;
+			return (DefaultValuePtr ? *DefaultValuePtr : DataType());
 		}
 
 		if constexpr(std::is_enum_v<DataType>)
@@ -58,7 +71,7 @@ struct TCameraContextDataReader
 				return *ActualValue;
 			}
 		}
-		return *DefaultValuePtr;
+		return (DefaultValuePtr ? *DefaultValuePtr : DataType());
 	}
 
 private:
diff --git a/Engine/Plugins/Cameras/GameplayCameras/Source/GameplayCameras/Public/Nodes/Blends/EasingBlendCameraNode.h b/Engine/Plugins/Cameras/GameplayCameras/Source/GameplayCameras/Public/Nodes/Blends/EasingBlendCameraNode.h
--- a/Engine/Plugins/Cameras/GameplayCameras/Source/GameplayCameras/Public/Nodes/Blends/EasingBlendCameraNode.h
+++ b/Engine/Plugins/Cameras/GameplayCameras/Source/GameplayCameras/Public/Nodes/Blends/EasingBlendCameraNode.h
@@ -68,7 +68,10 @@ protected:
 public:
 
 	/** The type of curve to use. */
-	UPROPERTY(EditAnywhere, Category=Blending)
+	UPROPERTY(EditAnywhere, Category=Blending, meta=(CameraContextData=true))
 	EEasingCameraBlendType BlendType = EEasingCameraBlendType::Linear;
+
+	UPROPERTY()
+	FCameraContextDataID BlendTypeDataID;
 };
 
