Cameras: fix some edge cases with the context data table

- Prevent access to some template functions for enum types when it can return pointers to incorrect underlying int size
- Add some extra checks to array setters

#Synced-CL 48619321

diff --git a/Engine/Plugins/Cameras/GameplayCameras/Source/GameplayCameras/Private/Core/CameraContextDataTable.cpp b/Engine/Plugins/Cameras/GameplayCameras/Source/GameplayCameras/Private/Core/CameraContextDataTable.cpp
--- a/Engine/Plugins/Cameras/GameplayCameras/Source/GameplayCameras/Private/Core/CameraContextDataTable.cpp
+++ b/Engine/Plugins/Cameras/GameplayCameras/Source/GameplayCameras/Private/Core/CameraContextDataTable.cpp
@@ -769,12 +769,19 @@ bool FCameraContextDataTable::TrySetArrayData(
 		return false;
 	}
 
-	if (!ensure(Entry->ContainerType == ECameraContextDataContainerType::Array))
+	if (!ensure(Entry->Type == ExpectedDataType && 
+				Entry->ContainerType == ECameraContextDataContainerType::Array &&
+				Entry->TypeObject == ExpectedDataTypeObject))
 	{
 		return false;
 	}
 
 	FArrayEntryHelper Helper(*Entry, Memory);
+	if (!ensure(Helper.IsValidIndex(Index)))
+	{
+		return false;
+	}
+
 	uint8* DataPtr = Helper.GetRawPtr(Index);
 	SetDataValue(Entry->Type, Entry->TypeObject, DataPtr, InRawDataPtr);
 
diff --git a/Engine/Plugins/Cameras/GameplayCameras/Source/GameplayCameras/Public/Core/CameraContextDataTable.h b/Engine/Plugins/Cameras/GameplayCameras/Source/GameplayCameras/Public/Core/CameraContextDataTable.h
--- a/Engine/Plugins/Cameras/GameplayCameras/Source/GameplayCameras/Public/Core/CameraContextDataTable.h
+++ b/Engine/Plugins/Cameras/GameplayCameras/Source/GameplayCameras/Public/Core/CameraContextDataTable.h
@@ -350,6 +350,9 @@ TSubclassOf<BaseClass> FCameraContextDataTable::GetClassData(FCameraContextDataI
 template<typename ValueType>
 const ValueType* FCameraContextDataTable::TryGetData(FCameraContextDataID InID) const
 {
+	// Can't return a direct pointer to an enum value because we coerce them to uint32 but ValueType may be smaller or larger.
+	static_assert(!std::is_enum_v<ValueType>, "Enum types aren't supported with this method.");
+
 	ECameraContextDataType DataType = TCameraContextDataTraits<ValueType>::GetDataType();
 	const UObject* DataTypeObject = TCameraContextDataTraits<ValueType>::GetDataTypeObject();
 	if (const uint8* RawValue = TryGetData(InID, DataType, DataTypeObject))
@@ -373,6 +376,9 @@ TConstArrayView<ValueType> FCameraContextDataTable::TryGetArrayData(FCameraConte
 template<typename ValueType>
 bool FCameraContextDataTable::TryGetArrayData(FCameraContextDataID InID, TConstArrayView<ValueType>& OutValues) const
 {
+	// Can't return a direct pointer to an array of enums because we coerce them to uint32 but ValueType may be smaller or larger.
+	static_assert(!std::is_enum_v<ValueType>, "Enum types aren't supported with this method.");
+
 	ECameraContextDataType DataType = TCameraContextDataTraits<ValueType>::GetDataType();
 	const UObject* DataTypeObject = TCameraContextDataTraits<ValueType>::GetDataTypeObject();
 	if (const FEntryScriptArray* Array = TryGetArrayData(InID, DataType, DataTypeObject))
@@ -410,7 +416,12 @@ void FCameraContextDataTable::SetStructData(FCameraContextDataID InID, const Str
 template<typename EnumType>
 void FCameraContextDataTable::SetEnumArrayData(FCameraContextDataID InID, TConstArrayView<EnumType> InData)
 {
-	SetArrayDataImpl(InID, ECameraContextDataType::Enum, StaticEnum<EnumType>(), InData);
+	TArray<uint32> CoercedData;
+	for (EnumType Item : InData)
+	{
+		CoercedData.Add((uint32)Item);
+	}
+	SetArrayDataImpl(InID, ECameraContextDataType::Enum, StaticEnum<EnumType>(), CoercedData);
 }
 
 template<typename StructType>
