Cameras: new camera action service

Camera actions are ways to execute operations on one or more running camera rig instances. They effectively "wrap" the root node evaluator for a camera rig, allowing arbitrary logic to be run before and after a camera rig evaluates. Camera actions also come with a few built-in options, such as the ability to clone themselves when a new camera rig activates, allowing a particular operation to persist over a long time.

#Synced-CL 48703199

diff --git a/Engine/Plugins/Cameras/GameplayCameras/Source/GameplayCameras/Private/Core/CameraSystemEvaluator.cpp b/Engine/Plugins/Cameras/GameplayCameras/Source/GameplayCameras/Private/Core/CameraSystemEvaluator.cpp
--- a/Engine/Plugins/Cameras/GameplayCameras/Source/GameplayCameras/Private/Core/CameraSystemEvaluator.cpp
+++ b/Engine/Plugins/Cameras/GameplayCameras/Source/GameplayCameras/Private/Core/CameraSystemEvaluator.cpp
@@ -18,6 +18,7 @@
 #include "Math/ColorList.h"
 #include "Serialization/MemoryReader.h"
 #include "Serialization/MemoryWriter.h"
+#include "Services/CameraActionService.h"
 #include "Services/CameraModifierService.h"
 #include "Services/CameraParameterSetterService.h"
 #include "Services/CameraShakeService.h"
@@ -97,6 +98,7 @@ void FCameraSystemEvaluator::Initialize(const FCameraSystemEvaluatorCreateParams
 	RegisterEvaluationService(MakeShared<FCameraParameterSetterService>());
 	RegisterEvaluationService(MakeShared<FCameraShakeService>());
 	RegisterEvaluationService(MakeShared<FOrientationInitializationService>());
+	RegisterEvaluationService(MakeShared<FCameraActionService>());
 
 	CameraRigCombinationRegistry = MakeShared<FCameraRigCombinationRegistry>();
 
@@ -239,7 +241,7 @@ void FCameraSystemEvaluator::UpdateImpl(float DeltaTime, ECameraNodeEvaluationTy
 	RootNodeResult.ContextDataTable.AutoResetValues();
 
 	// Pre-update all services.
-	PreUpdateServices(DeltaTime, ECameraEvaluationServiceFlags::None);
+	PreUpdateServices(DeltaTime, EvaluationType, ECameraEvaluationServiceFlags::None);
 
 	// Get the active evaluation context.
 	TSharedPtr<FCameraEvaluationContext> ActiveContext = ContextStack.GetActiveContext();
@@ -279,7 +281,7 @@ void FCameraSystemEvaluator::UpdateImpl(float DeltaTime, ECameraNodeEvaluationTy
 	}
 
 	// Post-update all services.
-	PostUpdateServices(DeltaTime, ECameraEvaluationServiceFlags::None);
+	PostUpdateServices(DeltaTime, EvaluationType, ECameraEvaluationServiceFlags::None);
 
 	// Harvest the result.
 	PreVisualResult.Reset(RootEvaluator->GetPreVisualLayerResult());
@@ -409,11 +411,12 @@ void FCameraSystemEvaluator::GetCombinedCameraRigRequest(TConstArrayView<FCamera
 	OutCombinedRequest.bForceActivateDeactivate = bAnyForceActivationDeactivation;
 }
 
-void FCameraSystemEvaluator::PreUpdateServices(float DeltaTime, ECameraEvaluationServiceFlags ExtraFlags)
+void FCameraSystemEvaluator::PreUpdateServices(float DeltaTime, ECameraNodeEvaluationType EvaluationType, ECameraEvaluationServiceFlags ExtraFlags)
 {
 	FCameraEvaluationServiceUpdateParams ServiceUpdateParams;
 	ServiceUpdateParams.Evaluator = this;
 	ServiceUpdateParams.DeltaTime = DeltaTime;
+	ServiceUpdateParams.EvaluationType = EvaluationType;
 
 	FCameraEvaluationServiceUpdateResult ServiceUpdateResult(RootNodeResult);
 
@@ -426,11 +429,12 @@ void FCameraSystemEvaluator::PreUpdateServices(float DeltaTime, ECameraEvaluatio
 	}
 }
 
-void FCameraSystemEvaluator::PostUpdateServices(float DeltaTime, ECameraEvaluationServiceFlags ExtraFlags)
+void FCameraSystemEvaluator::PostUpdateServices(float DeltaTime, ECameraNodeEvaluationType EvaluationType, ECameraEvaluationServiceFlags ExtraFlags)
 {
 	FCameraEvaluationServiceUpdateParams ServiceUpdateParams;
 	ServiceUpdateParams.Evaluator = this;
 	ServiceUpdateParams.DeltaTime = DeltaTime;
+	ServiceUpdateParams.EvaluationType = EvaluationType;
 
 	FCameraEvaluationServiceUpdateResult ServiceUpdateResult(RootNodeResult);
 
diff --git a/Engine/Plugins/Cameras/GameplayCameras/Source/GameplayCameras/Private/Core/DefaultRootCameraNode.cpp b/Engine/Plugins/Cameras/GameplayCameras/Source/GameplayCameras/Private/Core/DefaultRootCameraNode.cpp
--- a/Engine/Plugins/Cameras/GameplayCameras/Source/GameplayCameras/Private/Core/DefaultRootCameraNode.cpp
+++ b/Engine/Plugins/Cameras/GameplayCameras/Source/GameplayCameras/Private/Core/DefaultRootCameraNode.cpp
@@ -299,6 +299,11 @@ void FDefaultRootCameraNodeEvaluator::OnRunSingleCameraRig(const FSingleCameraRi
 	OutResult.bIsValid = true;
 }
 
+TSharedPtr<FCameraActionScope> FDefaultRootCameraNodeEvaluator::OnGetActiveCameraRigActionScope(bool bCreateIfNeeded)
+{
+	return MainLayer->GetActiveCameraRigActionScope(bCreateIfNeeded);
+}
+
 void FDefaultRootCameraNodeEvaluator::OnBlendStackEvent(const FBlendStackCameraRigEvent& InEvent)
 {
 	if (InEvent.EventType == EBlendStackCameraRigEventType::Pushed ||
diff --git a/Engine/Plugins/Cameras/GameplayCameras/Source/GameplayCameras/Private/Core/RootCameraNode.cpp b/Engine/Plugins/Cameras/GameplayCameras/Source/GameplayCameras/Private/Core/RootCameraNode.cpp
--- a/Engine/Plugins/Cameras/GameplayCameras/Source/GameplayCameras/Private/Core/RootCameraNode.cpp
+++ b/Engine/Plugins/Cameras/GameplayCameras/Source/GameplayCameras/Private/Core/RootCameraNode.cpp
@@ -100,6 +100,11 @@ void FRootCameraNodeEvaluator::RunSingleCameraRig(const FSingleCameraRigEvaluati
 	OnRunSingleCameraRig(Params, OutResult);
 }
 
+TSharedPtr<FCameraActionScope> FRootCameraNodeEvaluator::GetActiveCameraRigActionScope(bool bCreateIfNeeded)
+{
+	return OnGetActiveCameraRigActionScope(bCreateIfNeeded);
+}
+
 const FCameraNodeEvaluationResult& FRootCameraNodeEvaluator::GetPreVisualLayerResult() const
 {
 	return PreVisualResult;
diff --git a/Engine/Plugins/Cameras/GameplayCameras/Source/GameplayCameras/Private/Core/TransientBlendStackCameraNode.cpp b/Engine/Plugins/Cameras/GameplayCameras/Source/GameplayCameras/Private/Core/TransientBlendStackCameraNode.cpp
--- a/Engine/Plugins/Cameras/GameplayCameras/Source/GameplayCameras/Private/Core/TransientBlendStackCameraNode.cpp
+++ b/Engine/Plugins/Cameras/GameplayCameras/Source/GameplayCameras/Private/Core/TransientBlendStackCameraNode.cpp
@@ -12,8 +12,12 @@
 #include "Core/CameraRigCombinationRegistry.h"
 #include "Core/CameraRigTransition.h"
 #include "Core/CameraSystemEvaluator.h"
+#include "Debug/CameraDebugBlockBuilder.h"
+#include "Debug/CameraDebugRenderer.h"
 #include "Helpers/CameraRigTransitionFinder.h"
 #include "Nodes/Blends/PopBlendCameraNode.h"
+#include "Services/CameraActionEvaluator.h"
+#include "Services/CameraActionScope.h"
 #include "Services/CameraParameterSetterService.h"
 
 namespace UE::Cameras
@@ -219,6 +223,23 @@ void FTransientBlendStackCameraNodeEvaluator::FreezeAll(TSharedPtr<const FCamera
 	}
 }
 
+TSharedPtr<FCameraActionScope> FTransientBlendStackCameraNodeEvaluator::GetActiveCameraRigActionScope(bool bCreateIfNeeded)
+{
+	if (!EntryExtraInfos.IsEmpty())
+	{
+		FCameraRigEntryExtraInfo& ActiveEntryExtraInfo = EntryExtraInfos.Last();
+		if (bCreateIfNeeded && !ActiveEntryExtraInfo.ActionScope.IsValid())
+		{
+			ActiveEntryExtraInfo.ActionScope = MakeShared<FCameraActionScope>();
+
+			FCameraRigEvaluationInfo CameraRigInfo = GetActiveCameraRigEvaluationInfo();
+			ActiveEntryExtraInfo.ActionScope->Initialize(CameraRigInfo);
+		}
+		return ActiveEntryExtraInfo.ActionScope;
+	}
+	return nullptr;
+}
+
 void FTransientBlendStackCameraNodeEvaluator::OnInitialize(const FCameraNodeEvaluatorInitializeParams& Params, FCameraNodeEvaluationResult& OutResult)
 {
 	Super::OnInitialize(Params, OutResult);
@@ -272,7 +293,7 @@ void FTransientBlendStackCameraNodeEvaluator::InternalPreBlendPrepare(TArrayView
 		FCameraNodeEvaluationParams CurParams(Params);
 		CurParams.EvaluationContext = ResolvedEntry.Context;
 		CurParams.bIsFirstFrame = Entry.Flags.bIsFirstFrame;
-		CurParams.bIsActiveCameraRig = (ResolvedEntry.EntryIndex == Entries.Num() - 1);
+		CurParams.bIsActiveCameraRig = ResolvedEntry.bIsActiveEntry;
 
 		FCameraNodeEvaluationResult& CurResult(Entry.Result);
 
@@ -400,7 +421,7 @@ void FTransientBlendStackCameraNodeEvaluator::InternalUpdate(TArrayView<FResolve
 		FCameraNodeEvaluationParams CurParams(Params);
 		CurParams.EvaluationContext = ResolvedEntry.Context;
 		CurParams.bIsFirstFrame = Entry.Flags.bIsFirstFrame;
-		CurParams.bIsActiveCameraRig = (ResolvedEntry.EntryIndex == Entries.Num() - 1);
+		CurParams.bIsActiveCameraRig = ResolvedEntry.bIsActiveEntry;
 
 		// Start with the input given to us.
 		CurResult.Reset();
@@ -419,12 +440,23 @@ void FTransientBlendStackCameraNodeEvaluator::InternalUpdate(TArrayView<FResolve
 		CurResult.AddCameraPoseTrailPointIfNeeded(ContextResult.CameraPose.GetLocation());
 #endif  // WITH_EDITOR || UE_GAMEPLAY_CAMERAS_DEBUG
 
+		const bool bHasActionScope = EntryExtraInfo.ActionScope.IsValid();
+		if (bHasActionScope)
+		{
+			EntryExtraInfo.ActionScope->PreScopeRun(CurParams, CurResult);
+		}
+
 		// Run the camera rig's root node.
 		FCameraNodeEvaluator* RootEvaluator = Entry.RootEvaluator->GetRootEvaluator();
 		if (RootEvaluator)
 		{
 			RootEvaluator->Run(CurParams, CurResult);
 		}
+
+		if (bHasActionScope)
+		{
+			EntryExtraInfo.ActionScope->PostScopeRun(CurParams, CurResult);
+		}
 	}
 }
 
@@ -489,6 +521,12 @@ void FTransientBlendStackCameraNodeEvaluator::OnSerialize(const FCameraNodeEvalu
 
 	for (FCameraRigEntryExtraInfo& ExtraInfo : EntryExtraInfos)
 	{
+		if (ExtraInfo.ActionScope)
+		{
+			FCameraActionEvaluatorSerializeParams ActionScopeParams;
+			ExtraInfo.ActionScope->Serialize(ActionScopeParams, Ar);
+		}
+
 		Ar << ExtraInfo.bInputRunThisFrame;
 		Ar << ExtraInfo.bBlendRunThisFrame;
 		Ar << ExtraInfo.bHasPreBlendedParameters;
@@ -496,6 +534,19 @@ void FTransientBlendStackCameraNodeEvaluator::OnSerialize(const FCameraNodeEvalu
 	}
 }
 
+void FTransientBlendStackCameraNodeEvaluator::OnAddReferencedObjects(FReferenceCollector& Collector)
+{
+	Super::OnAddReferencedObjects(Collector);
+
+	for (FCameraRigEntryExtraInfo& ExtraInfo : EntryExtraInfos)
+	{
+		if (ExtraInfo.ActionScope)
+		{
+			ExtraInfo.ActionScope->AddReferencedObjects(Collector);
+		}
+	}
+}
+
 const UCameraRigTransition* FTransientBlendStackCameraNodeEvaluator::FindTransition(const FBlendStackCameraPushParams& Params) const
 {
 	// If we are forced to use a specific transition, our search is over.
@@ -611,6 +662,45 @@ const UCameraRigTransition* FTransientBlendStackCameraNodeEvaluator::FindTransit
 	return nullptr;
 }
 
+#if UE_GAMEPLAY_CAMERAS_DEBUG
+
+class FActionScopesCameraDebugBlock : public FCameraDebugBlock
+{
+	UE_DECLARE_CAMERA_DEBUG_BLOCK(, FActionScopesCameraDebugBlock)
+
+public:
+
+	FActionScopesCameraDebugBlock() = default;
+
+protected:
+
+	// FCameraDebugBlock interface.
+	virtual void OnDebugDraw(const FCameraDebugBlockDrawParams& Params, FCameraDebugRenderer& Renderer) override;
+};
+
+void FTransientBlendStackCameraNodeEvaluator::OnBuildDebugBlocks(const FCameraDebugBlockBuildParams& Params, FCameraDebugBlockBuilder& Builder)
+{
+	Super::OnBuildDebugBlocks(Params, Builder);
+
+	FActionScopesCameraDebugBlock& DebugBlock = Builder.AttachDebugBlock<FActionScopesCameraDebugBlock>();
+
+	for (const FCameraRigEntryExtraInfo& EntryExtraInfo : EntryExtraInfos)
+	{
+		if (EntryExtraInfo.ActionScope.IsValid())
+		{
+			EntryExtraInfo.ActionScope->BuildDebugBlocks(Params, Builder);
+		}
+	}
+}
+
+UE_DEFINE_CAMERA_DEBUG_BLOCK(FActionScopesCameraDebugBlock)
+
+void FActionScopesCameraDebugBlock::OnDebugDraw(const FCameraDebugBlockDrawParams& Params, FCameraDebugRenderer& Renderer)
+{
+}
+
+#endif
+
 #if WITH_EDITOR
 
 void FTransientBlendStackCameraNodeEvaluator::OnEntryReinitialized(int32 EntryIndex)
diff --git a/Engine/Plugins/Cameras/GameplayCameras/Source/GameplayCameras/Private/GameFramework/GameplayCameraComponentBase.cpp b/Engine/Plugins/Cameras/GameplayCameras/Source/GameplayCameras/Private/GameFramework/GameplayCameraComponentBase.cpp
--- a/Engine/Plugins/Cameras/GameplayCameras/Source/GameplayCameras/Private/GameFramework/GameplayCameraComponentBase.cpp
+++ b/Engine/Plugins/Cameras/GameplayCameras/Source/GameplayCameras/Private/GameFramework/GameplayCameraComponentBase.cpp
@@ -428,6 +428,54 @@ bool UGameplayCameraComponentBase::StopCameraShakeAsset(FCameraShakeInstanceID I
 	}
 }
 
+FCameraActionInstanceID UGameplayCameraComponentBase::StartAction(const UCameraAction* CameraAction)
+{
+	if (HasCameraSystem())
+	{
+		return IGameplayCameraSystemHost::StartAction(CameraAction);
+	}
+	else
+	{
+		UE_PRIVATE_GAMEPLAY_CAMERA_COMPONENT_ERROR_MESSAGE("Can't start action");
+		return FCameraActionInstanceID();
+	}
+}
+
+bool UGameplayCameraComponentBase::IsActionRunning(FCameraActionInstanceID InInstanceID)
+{
+	if (HasCameraSystem())
+	{
+		return IGameplayCameraSystemHost::IsActionRunning(InInstanceID);
+	}
+	return false;
+}
+
+bool UGameplayCameraComponentBase::StopAction(FCameraActionInstanceID InInstanceID)
+{
+	if (HasCameraSystem())
+	{
+		return IGameplayCameraSystemHost::StopAction(InInstanceID);
+	}
+	else
+	{
+		UE_PRIVATE_GAMEPLAY_CAMERA_COMPONENT_ERROR_MESSAGE("Can't stop action");
+		return false;
+	}
+}
+
+bool UGameplayCameraComponentBase::StopAllActionsOfClass(TSubclassOf<UCameraAction> InActionClass)
+{
+	if (HasCameraSystem())
+	{
+		return IGameplayCameraSystemHost::StopAllActionsOfClass(InActionClass);
+	}
+	else
+	{
+		UE_PRIVATE_GAMEPLAY_CAMERA_COMPONENT_ERROR_MESSAGE("Can't stop actions");
+		return false;
+	}
+}
+
 #undef UE_PRIVATE_GAMEPLAY_CAMERA_COMPONENT_VALIDATE_EVALUATION_CONTEXT
 
 FRotator UGameplayCameraComponentBase::GetEvaluatedCameraRotation() const
diff --git a/Engine/Plugins/Cameras/GameplayCameras/Source/GameplayCameras/Private/GameFramework/IGameplayCameraSystemHost.cpp b/Engine/Plugins/Cameras/GameplayCameras/Source/GameplayCameras/Private/GameFramework/IGameplayCameraSystemHost.cpp
--- a/Engine/Plugins/Cameras/GameplayCameras/Source/GameplayCameras/Private/GameFramework/IGameplayCameraSystemHost.cpp
+++ b/Engine/Plugins/Cameras/GameplayCameras/Source/GameplayCameras/Private/GameFramework/IGameplayCameraSystemHost.cpp
@@ -16,6 +16,7 @@
 #include "Engine/World.h"
 #include "GameFramework/GameplayCamerasPlayerCameraManager.h"
 #include "GameFramework/PlayerController.h"
+#include "Services/CameraActionService.h"
 #include "Services/CameraModifierService.h"
 #include "Services/CameraShakeService.h"
 
@@ -191,6 +192,54 @@ bool IGameplayCameraSystemHost::StopCameraShake(FCameraShakeInstanceID InInstanc
 	return false;
 }
 
+FCameraActionInstanceID IGameplayCameraSystemHost::StartAction(const UCameraAction* CameraAction)
+{
+	using namespace UE::Cameras;
+
+	if (ensure(CameraSystemEvaluator))
+	{
+		TSharedPtr<FCameraActionService> CameraActionService = CameraSystemEvaluator->FindEvaluationService<FCameraActionService>();
+		return CameraActionService->StartAction(CameraAction);
+	}
+	return FCameraActionInstanceID();
+}
+
+bool IGameplayCameraSystemHost::IsActionRunning(const FCameraActionInstanceID InInstanceID) const
+{
+	using namespace UE::Cameras;
+
+	if (ensure(CameraSystemEvaluator))
+	{
+		TSharedPtr<FCameraActionService> CameraActionService = CameraSystemEvaluator->FindEvaluationService<FCameraActionService>();
+		return CameraActionService->IsActionRunning(InInstanceID);
+	}
+	return false;
+}
+
+bool IGameplayCameraSystemHost::StopAction(const FCameraActionInstanceID InInstanceID)
+{
+	using namespace UE::Cameras;
+
+	if (ensure(CameraSystemEvaluator))
+	{
+		TSharedPtr<FCameraActionService> CameraActionService = CameraSystemEvaluator->FindEvaluationService<FCameraActionService>();
+		return CameraActionService->StopAction(InInstanceID);
+	}
+	return false;
+}
+
+bool IGameplayCameraSystemHost::StopAllActionsOfClass(TSubclassOf<UCameraAction> InActionClass)
+{
+	using namespace UE::Cameras;
+
+	if (ensure(CameraSystemEvaluator))
+	{
+		TSharedPtr<FCameraActionService> CameraActionService = CameraSystemEvaluator->FindEvaluationService<FCameraActionService>();
+		return CameraActionService->StopAllActionsOfClass(InActionClass);
+	}
+	return false;
+}
+
 IGameplayCameraSystemHost* IGameplayCameraSystemHost::FindActiveHost(APlayerController* PlayerController)
 {
 	if (PlayerController)
diff --git a/Engine/Plugins/Cameras/GameplayCameras/Source/GameplayCameras/Private/Services/CameraAction.cpp b/Engine/Plugins/Cameras/GameplayCameras/Source/GameplayCameras/Private/Services/CameraAction.cpp
new file mode 100644
--- /dev/null
+++ b/Engine/Plugins/Cameras/GameplayCameras/Source/GameplayCameras/Private/Services/CameraAction.cpp
@@ -0,0 +1,16 @@
+// Copyright Epic Games, Inc. All Rights Reserved.
+
+#include "Services/CameraAction.h"
+
+#include "Services/CameraActionEvaluator.h"
+
+#include UE_INLINE_GENERATED_CPP_BY_NAME(CameraAction)
+
+FCameraActionEvaluatorPtr UCameraAction::BuildEvaluator(FCameraActionEvaluatorBuilder& Builder) const
+{
+	using namespace UE::Cameras;
+	FCameraActionEvaluator* NewEvaluator = OnBuildEvaluator(Builder);
+	NewEvaluator->SetPrivateCameraAction(this);
+	return NewEvaluator;
+}
+
diff --git a/Engine/Plugins/Cameras/GameplayCameras/Source/GameplayCameras/Private/Services/CameraActionEvaluator.cpp b/Engine/Plugins/Cameras/GameplayCameras/Source/GameplayCameras/Private/Services/CameraActionEvaluator.cpp
new file mode 100644
--- /dev/null
+++ b/Engine/Plugins/Cameras/GameplayCameras/Source/GameplayCameras/Private/Services/CameraActionEvaluator.cpp
@@ -0,0 +1,68 @@
+// Copyright Epic Games, Inc. All Rights Reserved.
+
+#include "Services/CameraActionEvaluator.h"
+
+#include "Services/CameraAction.h"
+
+namespace UE::Cameras
+{
+
+UE_GAMEPLAY_CAMERAS_DEFINE_RTTI(FCameraActionEvaluator)
+
+void FCameraActionEvaluator::Initialize(const FCameraActionEvaluatorInitializeParams& Params, FCameraActionEvaluationResult& OutResult)
+{
+	OnInitialize(Params, OutResult);
+}
+
+void FCameraActionEvaluator::Teardown(const FCameraActionEvaluatorTeardownParams& Params)
+{
+	OnTeardown(Params);
+}
+
+void FCameraActionEvaluator::AddReferencedObjects(FReferenceCollector& Collector)
+{
+	Collector.AddReferencedObject(PrivateCameraAction);
+	OnAddReferencedObjects(Collector);
+}
+
+void FCameraActionEvaluator::PreScopeRun(const FCameraActionEvaluationParams& Params, FCameraActionEvaluationResult& OutResult)
+{
+	OnPreScopeRun(Params, OutResult);
+}
+
+void FCameraActionEvaluator::PostScopeRun(const FCameraActionEvaluationParams& Params, FCameraActionEvaluationResult& OutResult)
+{
+	OnPostScopeRun(Params, OutResult);
+}
+
+void FCameraActionEvaluator::Serialize(const FCameraActionEvaluatorSerializeParams& Params, FArchive& Ar)
+{
+	OnSerialize(Params, Ar);
+}
+
+void FCameraActionEvaluator::SetPrivateCameraAction(TObjectPtr<const UCameraAction> InCameraAction)
+{
+	PrivateCameraAction = InCameraAction;
+}
+
+void FCameraActionEvaluator::SetActionEvaluatorFlags(ECameraActionEvaluatorFlags InFlags)
+{
+	PrivateFlags = InFlags;
+}
+
+void FCameraActionEvaluator::AddActionEvaluatorFlags(ECameraActionEvaluatorFlags InFlags)
+{
+	PrivateFlags |= InFlags;
+}
+
+#if UE_GAMEPLAY_CAMERAS_DEBUG
+
+void FCameraActionEvaluator::BuildDebugBlocks(const FCameraDebugBlockBuildParams& Params, FCameraDebugBlockBuilder& Builder)
+{
+	OnBuildDebugBlocks(Params, Builder);
+}
+
+#endif  // UE_GAMEPLAY_CAMERAS_DEBUG
+
+}  // namespace UE::Cameras
+
diff --git a/Engine/Plugins/Cameras/GameplayCameras/Source/GameplayCameras/Private/Services/CameraActionScope.cpp b/Engine/Plugins/Cameras/GameplayCameras/Source/GameplayCameras/Private/Services/CameraActionScope.cpp
new file mode 100644
--- /dev/null
+++ b/Engine/Plugins/Cameras/GameplayCameras/Source/GameplayCameras/Private/Services/CameraActionScope.cpp
@@ -0,0 +1,313 @@
+// Copyright Epic Games, Inc. All Rights Reserved.
+
+#include "Services/CameraActionScope.h"
+
+#include "Core/CameraNodeEvaluator.h"
+#include "Core/CameraRigAsset.h"  // IWYU pragma: keep
+#include "Core/CameraRigEvaluationInfo.h"
+#include "Debug/CameraDebugBlockBuilder.h"
+#include "Debug/CameraDebugBlock.h"
+#include "Debug/CameraDebugRenderer.h"
+#include "Services/CameraAction.h"
+#include "Services/CameraActionEvaluator.h"
+#include "Services/CameraActionInstanceID.h"
+#include "Templates/SharedPointer.h"
+
+namespace UE::Cameras
+{
+
+FCameraActionInstanceID FCameraActionScope::StartAction(const UCameraAction* InCameraAction)
+{
+	if (InCameraAction)
+	{
+		FCameraActionEvaluatorBuilder Builder;
+		InCameraAction->BuildEvaluator(Builder);
+		if (TSharedPtr<FCameraActionEvaluator> Evaluator = Builder.BuiltEvaluator)
+		{
+			FActionInfo& NewAction = Actions.Emplace_GetRef();
+			NewAction.Evaluator = Evaluator;
+			NewAction.InstanceID = FCameraActionInstanceID(NextActionID++);
+			NewAction.State = EActionState::Uninitialized;
+			return NewAction.InstanceID;
+		}
+	}
+	return FCameraActionInstanceID();
+}
+
+bool FCameraActionScope::IsActionRunning(const FCameraActionInstanceID InInstanceID) const
+{
+	return Actions.ContainsByPredicate([InInstanceID](const FActionInfo& Item)
+			{
+				return Item.InstanceID == InInstanceID;
+			});
+}
+
+FCameraActionEvaluator* FCameraActionScope::GetAction(const FCameraActionInstanceID InInstanceID) const
+{
+	const FActionInfo* ActionInfo = Actions.FindByPredicate([InInstanceID](const FActionInfo& Item)
+			{
+				return Item.InstanceID == InInstanceID;
+			});
+	if (ActionInfo)
+	{
+		return ActionInfo->Evaluator.Get();
+	}
+	return nullptr;
+}
+
+bool FCameraActionScope::StopAction(const FCameraActionInstanceID InInstanceID)
+{
+	const int32 Index = Actions.IndexOfByPredicate([InInstanceID](const FActionInfo& Item)
+			{
+				return Item.InstanceID == InInstanceID;
+			});
+	if (Index != INDEX_NONE)
+	{
+		FActionInfo& ActionInfo = Actions[Index];
+		if (ActionInfo.State == EActionState::Running)
+		{
+			FCameraActionEvaluatorTeardownParams Params;
+			Params.Scope = SharedThis(this);
+			ActionInfo.Evaluator->Teardown(Params);
+		}
+		Actions.RemoveAt(Index);
+		return true;
+	}
+	return false;
+}
+
+void FCameraActionScope::Initialize(const FCameraRigEvaluationInfo& CameraRigInfo)
+{
+	// Don't store CameraRigInfo directly because we want to keep the evaluation context
+	// as a weak pointer.
+	CameraRigInstanceID = CameraRigInfo.InstanceID;
+	EvaluationContext = CameraRigInfo.EvaluationContext;
+	CameraRig = CameraRigInfo.CameraRig;
+	ChildResult = CameraRigInfo.LastResult;
+	ChildEvaluator = CameraRigInfo.RootEvaluator;
+}
+
+FCameraRigEvaluationInfo FCameraActionScope::GetCameraRigEvaluationInfo() const
+{
+	FCameraRigEvaluationInfo CameraRigEvaluationInfo(
+			CameraRigInstanceID,
+			EvaluationContext.Pin(),
+			CameraRig,
+			ChildResult,
+			ChildEvaluator);
+	return CameraRigEvaluationInfo;
+}
+
+void FCameraActionScope::AddReferencedObjects(FReferenceCollector& Collector)
+{
+	Collector.AddReferencedObject(CameraRig);
+
+	for (FActionInfo& ActionInfo : Actions)
+	{
+		ActionInfo.Evaluator->AddReferencedObjects(Collector);
+	}
+}
+
+void FCameraActionScope::PreScopeRun(const FCameraNodeEvaluationParams& Params, FCameraNodeEvaluationResult& OutResult)
+{
+	FCameraActionEvaluationParams ActionParams(Params);
+	ActionParams.Scope = SharedThis(this);
+
+	for (auto It = Actions.CreateIterator(); It; ++It)
+	{
+		FActionInfo& ActionInfo(*It);
+		FCameraActionEvaluationResult ActionResult(OutResult);
+		if (!PrepareActionForRun(ActionInfo, ActionResult))
+		{
+			continue;
+		}
+
+		ActionInfo.Evaluator->PreScopeRun(ActionParams, ActionResult);
+		if (ActionResult.bIsActionFinished)
+		{
+			ActionInfo.State = EActionState::Finished;
+			if (!Params.IsStatelessEvaluation())
+			{
+				It.RemoveCurrent();
+			}
+		}
+	}
+}
+
+void FCameraActionScope::PostScopeRun(const FCameraNodeEvaluationParams& Params, FCameraNodeEvaluationResult& OutResult)
+{
+	FCameraActionEvaluationParams ActionParams(Params);
+
+	for (auto It = Actions.CreateIterator(); It; ++It)
+	{
+		FActionInfo& ActionInfo(*It);
+		FCameraActionEvaluationResult ActionResult(OutResult);
+		if (!PrepareActionForRun(ActionInfo, ActionResult))
+		{
+			continue;
+		}
+
+		ActionInfo.Evaluator->PostScopeRun(ActionParams, ActionResult);
+		if (ActionResult.bIsActionFinished)
+		{
+			ActionInfo.State = EActionState::Finished;
+			if (!Params.IsStatelessEvaluation())
+			{
+				It.RemoveCurrent();
+			}
+		}
+	}
+}
+
+bool FCameraActionScope::PrepareActionForRun(FActionInfo& ActionInfo, FCameraActionEvaluationResult& OutResult)
+{
+	if (!ensure(ActionInfo.Evaluator))
+	{
+		return false;
+	}
+
+	if (ActionInfo.State == EActionState::Uninitialized)
+	{
+		FCameraActionEvaluatorInitializeParams InitParams;
+		InitParams.Scope = SharedThis(this);
+		ActionInfo.Evaluator->Initialize(InitParams, OutResult);
+		ActionInfo.State = EActionState::Running;
+	}
+
+	return ActionInfo.State == EActionState::Running;
+}
+
+void FCameraActionScope::Serialize(const FCameraActionEvaluatorSerializeParams& Params, FArchive& Ar)
+{
+	int32 NumActionsToSerialize = Actions.Num();
+
+	if (Ar.IsSaving())
+	{
+		int32 NumActions = Actions.Num();
+		Ar << NumActions;
+	}
+	else if (Ar.IsLoading())
+	{
+		int32 LoadedNumActions = 0;
+		Ar << LoadedNumActions;
+
+		ensureMsgf(
+				LoadedNumActions == Actions.Num(),
+				TEXT("The number of actions changed since this action scope was serialized!"));
+		NumActionsToSerialize = LoadedNumActions;
+	}
+
+	for (int32 Index = 0; Index < NumActionsToSerialize; ++Index)
+	{
+		FActionInfo& ActionInfo(Actions[Index]);
+		ActionInfo.Evaluator->Serialize(Params, Ar);
+		Ar << ActionInfo.InstanceID;
+
+		uint8 RawState = (uint8)ActionInfo.State;
+		Ar << RawState;
+		ActionInfo.State = (EActionState)RawState;
+	}
+}
+
+#if UE_GAMEPLAY_CAMERAS_DEBUG
+
+class FCameraActionScopeDebugBlock : public FCameraDebugBlock
+{
+	UE_DECLARE_CAMERA_DEBUG_BLOCK(, FCameraActionScopeDebugBlock)
+
+public:
+
+	FCameraActionScopeDebugBlock() = default;
+	FCameraActionScopeDebugBlock(const FString& InScopeCameraRigName);
+
+protected:
+
+	// FCameraDebugBlock interface.
+	virtual void OnDebugDraw(const FCameraDebugBlockDrawParams& Params, FCameraDebugRenderer& Renderer) override;
+	virtual void OnSerialize(FArchive& Ar) override;
+
+private:
+
+	FString ScopeCameraRigName;
+};
+
+class FCameraActionDebugBlock : public FCameraDebugBlock
+{
+	UE_DECLARE_CAMERA_DEBUG_BLOCK(, FCameraActionDebugBlock)
+
+public:
+
+	FCameraActionDebugBlock() = default;
+	FCameraActionDebugBlock(const UCameraAction* InCameraAction);
+
+protected:
+
+	// FCameraDebugBlock interface.
+	virtual void OnDebugDraw(const FCameraDebugBlockDrawParams& Params, FCameraDebugRenderer& Renderer) override;
+	virtual void OnSerialize(FArchive& Ar) override;
+
+private:
+
+	FString ActionClassName;
+};
+
+void FCameraActionScope::BuildDebugBlocks(const FCameraDebugBlockBuildParams& Params, FCameraDebugBlockBuilder& Builder)
+{
+	const FString CameraRigName = GetNameSafe(CameraRig);
+	FCameraActionScopeDebugBlock& DebugBlock = Builder.StartChildDebugBlock<FCameraActionScopeDebugBlock>(CameraRigName);
+	{
+		for (const FActionInfo& ActionInfo : Actions)
+		{
+			if (ActionInfo.Evaluator)
+			{
+				const UCameraAction* ActionData = ActionInfo.Evaluator->GetCameraAction();
+				Builder.StartChildDebugBlock<FCameraActionDebugBlock>(ActionData);
+				{
+					ActionInfo.Evaluator->BuildDebugBlocks(Params, Builder);
+				}
+				Builder.EndChildDebugBlock();
+			}
+		}
+	}
+	Builder.EndChildDebugBlock();
+}
+
+UE_DEFINE_CAMERA_DEBUG_BLOCK(FCameraActionScopeDebugBlock)
+
+FCameraActionScopeDebugBlock::FCameraActionScopeDebugBlock(const FString& InScopeCameraRigName)
+	: ScopeCameraRigName(InScopeCameraRigName)
+{
+}
+
+void FCameraActionScopeDebugBlock::OnDebugDraw(const FCameraDebugBlockDrawParams& Params, FCameraDebugRenderer& Renderer)
+{
+	Renderer.AddText(TEXT("%d action(s) for {cam_notice}%s{cam_default}"), GetChildren().Num(), *ScopeCameraRigName);
+}
+
+void FCameraActionScopeDebugBlock::OnSerialize(FArchive& Ar)
+{
+	Ar << ScopeCameraRigName;
+}
+
+
+UE_DEFINE_CAMERA_DEBUG_BLOCK(FCameraActionDebugBlock)
+
+FCameraActionDebugBlock::FCameraActionDebugBlock(const UCameraAction* InCameraAction)
+{
+	ActionClassName = (InCameraAction ? InCameraAction->GetClass()->GetName() : TEXT("<no action>"));
+}
+
+void FCameraActionDebugBlock::OnDebugDraw(const FCameraDebugBlockDrawParams& Params, FCameraDebugRenderer& Renderer)
+{
+	Renderer.AddText(TEXT("{cam_passive}[%s]{cam_default} "), *ActionClassName);
+}
+
+void FCameraActionDebugBlock::OnSerialize(FArchive& Ar)
+{
+	Ar << ActionClassName;
+}
+
+#endif  // UE_GAMEPLAY_CAMERAS_DEBUG
+
+}  // namespace UE::Cameras
+
diff --git a/Engine/Plugins/Cameras/GameplayCameras/Source/GameplayCameras/Private/Services/CameraActionService.cpp b/Engine/Plugins/Cameras/GameplayCameras/Source/GameplayCameras/Private/Services/CameraActionService.cpp
new file mode 100644
--- /dev/null
+++ b/Engine/Plugins/Cameras/GameplayCameras/Source/GameplayCameras/Private/Services/CameraActionService.cpp
@@ -0,0 +1,370 @@
+// Copyright Epic Games, Inc. All Rights Reserved.
+
+#include "Services/CameraActionService.h"
+
+#include "Containers/AllowShrinking.h"
+#include "Core/CameraRigAsset.h"  // IWYU pragma: keep
+#include "Core/CameraSystemEvaluator.h"
+#include "Core/RootCameraNode.h"
+#include "Core/RootCameraNodeCameraRigEvent.h"
+#include "Debug/CameraDebugBlock.h"
+#include "Debug/CameraDebugBlockBuilder.h"
+#include "Debug/CameraDebugRenderer.h"
+#include "Serialization/MemoryReader.h"
+#include "Serialization/MemoryWriter.h"
+#include "Services/CameraAction.h"
+#include "Services/CameraActionEvaluator.h"
+#include "Services/CameraActionScope.h"
+
+namespace UE::Cameras
+{
+
+UE_DEFINE_CAMERA_EVALUATION_SERVICE(FCameraActionService)
+
+FCameraActionService::~FCameraActionService()
+{
+	// If we are destroying the owning evaluator while a cloning camera action is running,
+	// we need to clean-up our delegate for root node events.
+	if (RootNodeEvaluator && NumCloningActions > 0)
+	{
+		RootNodeEvaluator->OnCameraRigEvent().RemoveAll(this);
+	}
+}
+
+void FCameraActionService::OnInitialize(const FCameraEvaluationServiceInitializeParams& Params)
+{
+	SetEvaluationServiceFlags(ECameraEvaluationServiceFlags::NeedsPostUpdate);
+
+	RootNodeEvaluator = Params.Evaluator->GetRootNodeEvaluator();
+}
+
+void FCameraActionService::OnPostUpdate(const FCameraEvaluationServiceUpdateParams& Params, FCameraEvaluationServiceUpdateResult& OutResult)
+{
+	if (Params.EvaluationType == ECameraNodeEvaluationType::Standard)
+	{
+		CleanUpActions();
+	}
+}
+
+FCameraActionInstanceID FCameraActionService::StartAction(const UCameraAction* CameraAction)
+{
+	if (!CameraAction)
+	{
+		return FCameraActionInstanceID();
+	}
+	
+	TSharedPtr<FCameraActionScope> ActionScope = RootNodeEvaluator->GetActiveCameraRigActionScope(true);
+	if (!ensure(ActionScope))
+	{
+		return FCameraActionInstanceID();
+	}
+	
+	FCameraActionInstanceID NewInstanceID = ActionScope->StartAction(CameraAction);
+	if (!ensure(NewInstanceID.IsValid()))
+	{
+		return FCameraActionInstanceID();
+	}
+
+	FActionSetInfo& NewActionSet = ActionSets.Emplace_GetRef();
+	NewActionSet.Data = CameraAction;
+	NewActionSet.Actions.Add(FActionInfo{ ActionScope, NewInstanceID });
+	NewActionSet.InstanceID = FCameraActionInstanceID(NextActionSetID++);
+
+	if (CameraAction->bPropagateToNewCameraRigs)
+	{
+		AddCloningAction();
+	}
+	
+	return NewActionSet.InstanceID;
+}
+
+bool FCameraActionService::IsActionRunning(const FCameraActionInstanceID InInstanceID) const
+{
+	if (InInstanceID.IsValid())
+	{
+		return ActionSets.ContainsByPredicate([InInstanceID](const FActionSetInfo& Item)
+				{
+					return Item.InstanceID == InInstanceID;
+				});
+	}
+	return false;
+}
+
+bool FCameraActionService::StopAction(const FCameraActionInstanceID InInstanceID)
+{
+	const int32 Index = ActionSets.IndexOfByPredicate([InInstanceID](const FActionSetInfo& Item)
+			{
+				return Item.InstanceID == InInstanceID;
+			});
+	if (Index != INDEX_NONE)
+	{
+		FActionSetInfo& ActionSet = ActionSets[Index];
+		for (FActionInfo& Action : ActionSet.Actions)
+		{
+			TSharedPtr<FCameraActionScope> ActionScope = Action.ActionScope.Pin();
+			if (ensure(ActionScope && Action.ActionID.IsValid()))
+			{
+				const bool bStopped = ActionScope->StopAction(Action.ActionID);
+				ensure(bStopped);
+			}
+		}
+		if (ActionSet.Data->bPropagateToNewCameraRigs)
+		{
+			RemoveCloningAction();
+		}
+		ActionSets.RemoveAt(Index);
+		return true;
+	}
+	return false;
+}
+
+bool FCameraActionService::StopAllActionsOfClass(TSubclassOf<UCameraAction> InActionClass)
+{
+	bool bAnyStopped = false;
+	for (auto It = ActionSets.CreateIterator(); It; ++It)
+	{
+		FActionSetInfo& ActionSet(*It);
+		if (ActionSet.Data && ActionSet.Data->IsA(InActionClass))
+		{
+			for (FActionInfo& Action : ActionSet.Actions)
+			{
+				if (TSharedPtr<FCameraActionScope> ActionScope = Action.ActionScope.Pin())
+				{
+					const bool bStopped = ActionScope->StopAction(Action.ActionID);
+					ensure(bStopped);
+					bAnyStopped = true;
+				}
+			}
+			if (ActionSet.Data->bPropagateToNewCameraRigs)
+			{
+				RemoveCloningAction();
+			}
+			It.RemoveCurrent();
+		}
+	}
+	return bAnyStopped;
+}
+
+void FCameraActionService::OnRootCameraNodeCameraRigEvent(const FRootCameraNodeCameraRigEvent& InEvent)
+{
+	if (InEvent.EventType == ERootCameraNodeCameraRigEventType::Activated && InEvent.EventLayer == ECameraRigLayer::Main)
+	{
+		TSharedPtr<FCameraActionScope> ActionScope = RootNodeEvaluator->GetActiveCameraRigActionScope(true);
+		if (ensure(ActionScope))
+		{
+			int32 NumCloned = 0;
+			for (FActionSetInfo& ActionSet : ActionSets)
+			{
+				if (ensure(ActionSet.Data) && ActionSet.Data->bPropagateToNewCameraRigs)
+				{
+					CloneAction(ActionSet, ActionScope.ToSharedRef());
+					++NumCloned;
+				}
+			}
+			ensureMsgf(
+					NumCloned > 0, 
+					TEXT("No actions found to clone: we are still registered with root camera rig events for nothing"));
+		}
+	}
+}
+
+void FCameraActionService::AddCloningAction()
+{
+	const bool bRegisterEvent = (NumCloningActions == 0);
+	++NumCloningActions;
+	if (bRegisterEvent)
+	{
+		RootNodeEvaluator->OnCameraRigEvent().AddRaw(this, &FCameraActionService::OnRootCameraNodeCameraRigEvent);
+	}
+}
+
+void FCameraActionService::RemoveCloningAction()
+{
+	ensure(NumCloningActions > 0);
+	--NumCloningActions;
+	if (NumCloningActions == 0)
+	{
+		RootNodeEvaluator->OnCameraRigEvent().RemoveAll(this);
+	}
+}
+
+void FCameraActionService::CloneAction(FActionSetInfo& ActionSet, TSharedRef<FCameraActionScope> NewActionScope)
+{
+	if (!ensure(ActionSet.Actions.Num() > 0))
+	{
+		return;
+	}
+
+	const FActionInfo& ActiveAction = ActionSet.Actions.Last();
+	TSharedPtr<FCameraActionScope> ActiveActionScope = ActiveAction.ActionScope.Pin();
+	if (!ensure(ActiveActionScope))
+	{
+		return;
+	}
+
+	FCameraActionEvaluator* ActiveEvaluator = ActiveActionScope->GetAction(ActiveAction.ActionID);
+	if (!ensure(ActiveEvaluator))
+	{
+		return;
+	}
+
+	TArray<uint8> EvaluatorSnapshot;
+	FCameraActionEvaluatorSerializeParams Params;
+	{
+		FMemoryWriter Writer(EvaluatorSnapshot);
+		ActiveEvaluator->Serialize(Params, Writer);
+	}
+
+	const FCameraActionInstanceID NewInstanceID = NewActionScope->StartAction(ActionSet.Data);
+	FCameraActionEvaluator* NewEvaluator = NewActionScope->GetAction(NewInstanceID);
+	if (ensure(NewEvaluator))
+	{
+		FMemoryReader Reader(EvaluatorSnapshot);
+		NewEvaluator->Serialize(Params, Reader);
+
+		ActionSet.Actions.Add(FActionInfo{ NewActionScope, NewInstanceID });
+	}
+}
+
+void FCameraActionService::CleanUpActions()
+{
+	for (auto SetIt = ActionSets.CreateIterator(); SetIt; ++SetIt)
+	{
+		FActionSetInfo& ActionSet(*SetIt);
+		for (auto ActionIt = ActionSet.Actions.CreateIterator(); ActionIt; ++ActionIt)
+		{
+			// Remove finished actions, whether they finished on their own, or were ended by their action scope 
+			// being discarded.
+			FActionInfo& ActionInfo(*ActionIt);
+			if (TSharedPtr<FCameraActionScope> ActionScope = ActionInfo.ActionScope.Pin())
+			{
+				if (!ActionScope->IsActionRunning(ActionInfo.ActionID))
+				{
+					ActionIt.RemoveCurrent(EAllowShrinking::No);
+				}
+			}
+			else
+			{
+				ActionIt.RemoveCurrent(EAllowShrinking::No);
+			}
+		}
+		if (ActionSet.Actions.IsEmpty())
+		{
+			if (ActionSet.Data->bPropagateToNewCameraRigs)
+			{
+				RemoveCloningAction();
+			}
+			SetIt.RemoveCurrent(EAllowShrinking::No);
+		}
+	}
+}
+
+#if UE_GAMEPLAY_CAMERAS_DEBUG
+
+class FCameraActionServiceDebugBlock : public FCameraDebugBlock
+{
+	UE_DECLARE_CAMERA_DEBUG_BLOCK(, FCameraActionServiceDebugBlock)
+
+public:
+
+	FCameraActionServiceDebugBlock() = default;
+	FCameraActionServiceDebugBlock(const FCameraActionService& InService);
+
+protected:
+
+	// FCameraDebugBlock interface.
+	virtual void OnDebugDraw(const FCameraDebugBlockDrawParams& Params, FCameraDebugRenderer& Renderer) override;
+	virtual void OnSerialize(FArchive& Ar) override;
+
+private:
+
+	struct FActionDebugInfo
+	{
+		FString ScopeCameraRigName;
+		FCameraActionInstanceID ActionID;
+	};
+
+	struct FActionSetDebugInfo
+	{
+		FString ActionName;
+		TArray<FActionDebugInfo> Actions;
+		FCameraActionInstanceID InstanceID;
+	};
+
+	friend FArchive& operator<< (FArchive&, FActionDebugInfo&);
+	friend FArchive& operator<< (FArchive&, FActionSetDebugInfo&);
+
+	TArray<FActionSetDebugInfo> ActionSets;
+};
+
+FArchive& operator<< (FArchive& Ar, FCameraActionServiceDebugBlock::FActionDebugInfo& Item)
+{
+	Ar << Item.ScopeCameraRigName;
+	Ar << Item.ActionID;
+	return Ar;
+}
+
+FArchive& operator<< (FArchive& Ar, FCameraActionServiceDebugBlock::FActionSetDebugInfo& Item)
+{
+	Ar << Item.ActionName;
+	Ar << Item.Actions;
+	Ar << Item.InstanceID;
+	return Ar;
+}
+
+UE_DEFINE_CAMERA_DEBUG_BLOCK(FCameraActionServiceDebugBlock)
+
+FCameraActionServiceDebugBlock::FCameraActionServiceDebugBlock(const FCameraActionService& InService)
+{
+	for (const FCameraActionService::FActionSetInfo& ActionSet : InService.ActionSets)
+	{
+		FActionSetDebugInfo& ActionSetDebugInfo = ActionSets.Emplace_GetRef();
+		ActionSetDebugInfo.ActionName = GetNameSafe(ActionSet.Data);
+		ActionSetDebugInfo.InstanceID = ActionSet.InstanceID;
+
+		for (const FCameraActionService::FActionInfo& Action : ActionSet.Actions)
+		{
+			FActionDebugInfo& ActionDebugInfo = ActionSetDebugInfo.Actions.Emplace_GetRef();
+			ActionDebugInfo.ActionID = Action.ActionID;
+			if (TSharedPtr<const FCameraActionScope> ActionScope = Action.ActionScope.Pin())
+			{
+				ActionDebugInfo.ScopeCameraRigName = GetNameSafe(ActionScope->GetCameraRig());
+			}
+			else
+			{
+				ActionDebugInfo.ScopeCameraRigName = TEXT("None");
+			}
+		}
+	}
+}
+
+void FCameraActionServiceDebugBlock::OnDebugDraw(const FCameraDebugBlockDrawParams& Params, FCameraDebugRenderer& Renderer)
+{
+	for (const FActionSetDebugInfo& ActionSet : ActionSets)
+	{
+		Renderer.AddText(TEXT("Action Set: {cam_notice}%s{cam_default}\n"), *ActionSet.ActionName);
+		Renderer.AddIndent();
+		{
+			for (const FActionDebugInfo& Action : ActionSet.Actions)
+			{
+				Renderer.AddText(TEXT("Action Scope {cam_notice2}%s{cam_default}\n"), *Action.ScopeCameraRigName);
+			}
+		}
+		Renderer.RemoveIndent();
+	}
+}
+
+void FCameraActionServiceDebugBlock::OnSerialize(FArchive& Ar)
+{
+	Ar << ActionSets;
+}
+
+void FCameraActionService::OnBuildDebugBlocks(const FCameraDebugBlockBuildParams& Params, FCameraDebugBlockBuilder& Builder)
+{
+	Builder.AttachDebugBlock<FCameraActionServiceDebugBlock>(*this);
+}
+
+#endif  // UE_GAMEPLAY_CAMERAS_DEBUG
+
+} // namespace UE::Cameras
+
diff --git a/Engine/Plugins/Cameras/GameplayCameras/Source/GameplayCameras/Public/Core/CameraEvaluationService.h b/Engine/Plugins/Cameras/GameplayCameras/Source/GameplayCameras/Public/Core/CameraEvaluationService.h
--- a/Engine/Plugins/Cameras/GameplayCameras/Source/GameplayCameras/Public/Core/CameraEvaluationService.h
+++ b/Engine/Plugins/Cameras/GameplayCameras/Source/GameplayCameras/Public/Core/CameraEvaluationService.h
@@ -2,6 +2,7 @@
 
 #pragma once
 
+#include "Core/CameraNodeEvaluator.h"
 #include "Core/CameraObjectRtti.h"
 #include "Debug/CameraDebugBlockFwd.h"
 #include "GameplayCameras.h"
@@ -45,6 +46,7 @@ struct FCameraEvaluationServiceUpdateParams
 {
 	FCameraSystemEvaluator* Evaluator = nullptr;
 	float DeltaTime = 0.f;
+	ECameraNodeEvaluationType EvaluationType = ECameraNodeEvaluationType::Standard;
 };
 
 /** Result structure for updating an evaluation service. */
diff --git a/Engine/Plugins/Cameras/GameplayCameras/Source/GameplayCameras/Public/Core/CameraSystemEvaluator.h b/Engine/Plugins/Cameras/GameplayCameras/Source/GameplayCameras/Public/Core/CameraSystemEvaluator.h
--- a/Engine/Plugins/Cameras/GameplayCameras/Source/GameplayCameras/Public/Core/CameraSystemEvaluator.h
+++ b/Engine/Plugins/Cameras/GameplayCameras/Source/GameplayCameras/Public/Core/CameraSystemEvaluator.h
@@ -301,8 +301,8 @@ private:
 	void UpdateCameraDirector(float DeltaTime, FCameraDirectorEvaluator* CameraDirectorEvaluator);
 	void GetCombinedCameraRigRequest(TConstArrayView<FCameraRigActivationDeactivationRequest> Requests, FCameraRigActivationDeactivationRequest& OutCombinedRequest);
 
-	void PreUpdateServices(float DeltaTime, ECameraEvaluationServiceFlags ExtraFlags);
-	void PostUpdateServices(float DeltaTime, ECameraEvaluationServiceFlags ExtraFlags);
+	void PreUpdateServices(float DeltaTime, ECameraNodeEvaluationType EvaluationType, ECameraEvaluationServiceFlags ExtraFlags);
+	void PostUpdateServices(float DeltaTime, ECameraNodeEvaluationType EvaluationType, ECameraEvaluationServiceFlags ExtraFlags);
 
 	void NotifyRootCameraNodeEvent(const FRootCameraNodeCameraRigEvent& InEvent);
 
diff --git a/Engine/Plugins/Cameras/GameplayCameras/Source/GameplayCameras/Public/Core/DefaultRootCameraNode.h b/Engine/Plugins/Cameras/GameplayCameras/Source/GameplayCameras/Public/Core/DefaultRootCameraNode.h
--- a/Engine/Plugins/Cameras/GameplayCameras/Source/GameplayCameras/Public/Core/DefaultRootCameraNode.h
+++ b/Engine/Plugins/Cameras/GameplayCameras/Source/GameplayCameras/Public/Core/DefaultRootCameraNode.h
@@ -78,6 +78,7 @@ protected:
 	virtual const FCameraVariableTable* OnGetBlendedParameters() const override;
 	virtual void OnBuildSingleCameraRigHierarchy(const FSingleCameraRigHierarchyBuildParams& Params, FCameraNodeEvaluatorHierarchy& OutHierarchy) override;
 	virtual void OnRunSingleCameraRig(const FSingleCameraRigEvaluationParams& Params, FCameraNodeEvaluationResult& OutResult) override;
+	virtual TSharedPtr<FCameraActionScope> OnGetActiveCameraRigActionScope(bool bCreateIfNeeded) override;
 
 private:
 
diff --git a/Engine/Plugins/Cameras/GameplayCameras/Source/GameplayCameras/Public/Core/RootCameraNode.h b/Engine/Plugins/Cameras/GameplayCameras/Source/GameplayCameras/Public/Core/RootCameraNode.h
--- a/Engine/Plugins/Cameras/GameplayCameras/Source/GameplayCameras/Public/Core/RootCameraNode.h
+++ b/Engine/Plugins/Cameras/GameplayCameras/Source/GameplayCameras/Public/Core/RootCameraNode.h
@@ -26,6 +26,7 @@ class URootCameraNode : public UCameraNode
 namespace UE::Cameras
 {
 
+class FCameraActionScope;
 class FCameraEvaluationContext;
 class FCameraSystemEvaluator;
 struct FCameraRigActivationDeactivationRequest;
@@ -163,6 +164,9 @@ public:
 	 */
 	GAMEPLAYCAMERAS_API void RunSingleCameraRig(const FSingleCameraRigEvaluationParams& Params, FCameraNodeEvaluationResult& OutResult);
 
+	/** Gets the action scope for the active camera rig in the main layer. */
+	GAMEPLAYCAMERAS_API TSharedPtr<FCameraActionScope> GetActiveCameraRigActionScope(bool bCreateIfNeeded = false);
+
 	/** Gets the delegate for camera rig events. */
 	FOnRootCameraNodeCameraRigEvent& OnCameraRigEvent() { return OnCameraRigEventDelegate; }
 
@@ -200,6 +204,9 @@ protected:
 	/** Evaluates a single camera rig. See comments on RunSingleCameraRig. */
 	virtual void OnRunSingleCameraRig(const FSingleCameraRigEvaluationParams& Params, FCameraNodeEvaluationResult& OutResult) {}
 
+	/** Gets the action scope for the active camera rig in the main layer. */
+	virtual TSharedPtr<FCameraActionScope> OnGetActiveCameraRigActionScope(bool bCreateIfNeeded) { return nullptr; }
+
 protected:
 
 	GAMEPLAYCAMERAS_API void SetPreVisualLayerResult(const FCameraNodeEvaluationResult& InResult);
diff --git a/Engine/Plugins/Cameras/GameplayCameras/Source/GameplayCameras/Public/Core/TransientBlendStackCameraNode.h b/Engine/Plugins/Cameras/GameplayCameras/Source/GameplayCameras/Public/Core/TransientBlendStackCameraNode.h
--- a/Engine/Plugins/Cameras/GameplayCameras/Source/GameplayCameras/Public/Core/TransientBlendStackCameraNode.h
+++ b/Engine/Plugins/Cameras/GameplayCameras/Source/GameplayCameras/Public/Core/TransientBlendStackCameraNode.h
@@ -9,6 +9,8 @@
 namespace UE::Cameras
 {
 
+class FCameraActionScope;
+class FCameraActionService;
 class FCameraParameterSetterService;
 
 /**
@@ -70,12 +72,20 @@ public:
 	/** Gets the variable table containing the blended camera object interfaces parameters. */
 	const FCameraVariableTable& GetBlendedParameters() const { return PreBlendVariableTable; }
 
+	/** Gets the active camera rig's action scope, optionally creating one if needed. */
+	TSharedPtr<FCameraActionScope> GetActiveCameraRigActionScope(bool bCreateIfNeeded = false);
+
 protected:
 
 	// FCameraNodeEvaluator interface.
 	virtual void OnInitialize(const FCameraNodeEvaluatorInitializeParams& Params, FCameraNodeEvaluationResult& OutResult) override;
 	virtual void OnRun(const FCameraNodeEvaluationParams& Params, FCameraNodeEvaluationResult& OutResult) override;
 	virtual void OnSerialize(const FCameraNodeEvaluatorSerializeParams& Params, FArchive& Ar) override;
+	virtual void OnAddReferencedObjects(FReferenceCollector& Collector) override;
+
+#if UE_GAMEPLAY_CAMERAS_DEBUG
+	virtual void OnBuildDebugBlocks(const FCameraDebugBlockBuildParams& Params, FCameraDebugBlockBuilder& Builder) override;
+#endif
 
 #if WITH_EDITOR
 	// FBlendStackCameraNodeEvaluator interface.
@@ -101,6 +111,9 @@ private:
 	/** Extra blending-related info for each camera rig entry. */
 	struct FCameraRigEntryExtraInfo
 	{
+		/** Action scope for this entry. */
+		TSharedPtr<FCameraActionScope> ActionScope;
+
 		/** Whether input slots were run. */
 		bool bInputRunThisFrame = false;
 		/** Whether the blend node was run. */
diff --git a/Engine/Plugins/Cameras/GameplayCameras/Source/GameplayCameras/Public/GameFramework/GameplayCameraComponentBase.h b/Engine/Plugins/Cameras/GameplayCameras/Source/GameplayCameras/Public/GameFramework/GameplayCameraComponentBase.h
--- a/Engine/Plugins/Cameras/GameplayCameras/Source/GameplayCameras/Public/GameFramework/GameplayCameraComponentBase.h
+++ b/Engine/Plugins/Cameras/GameplayCameras/Source/GameplayCameras/Public/GameFramework/GameplayCameraComponentBase.h
@@ -5,9 +5,11 @@
 #include "Components/SceneComponent.h"
 #include "Core/CameraAssetReference.h"
 #include "Core/CameraEvaluationContext.h"
+#include "Core/CameraRigInstanceID.h"
 #include "Core/CameraShakeInstanceID.h"
 #include "GameFramework/BlueprintCameraEvaluationDataRef.h"
 #include "GameFramework/IGameplayCameraSystemHost.h"
+#include "Services/CameraActionInstanceID.h"
 #include "UObject/ObjectMacros.h"
 
 #include "GameplayCameraComponentBase.generated.h"
@@ -16,6 +18,7 @@ class APlayerController;
 class FPrimitiveDrawInterface;
 class FSceneView;
 class FViewport;
+class UCameraAction;
 class UCameraAsset;
 class UCameraShakeAsset;
 class UCanvas;
@@ -146,6 +149,24 @@ public:
 	UFUNCTION(BlueprintCallable, Category="Camera")
 	bool StopCameraShakeAsset(FCameraShakeInstanceID InInstanceID, bool bImmediately = false);
 
+public:
+
+	/** Starts the given camera action. */
+	UFUNCTION(BlueprintCallable, Category="Camera")
+	FCameraActionInstanceID StartAction(const UCameraAction* CameraAction);
+
+	/** Returns whether the given camera action instance is still running. */
+	UFUNCTION(BlueprintCallable, Category="Camera")
+	bool IsActionRunning(FCameraActionInstanceID InInstanceID);
+
+	/** Stops the given camera action instance. */
+	UFUNCTION(BlueprintCallable, Category="Camera")
+	bool StopAction(FCameraActionInstanceID InInstanceID);
+
+	/** Stops all camera actions of a given class. */
+	UFUNCTION(BlueprintCallable, Category="Camera")
+	bool StopAllActionsOfClass(TSubclassOf<UCameraAction> InActionClass);
+
 public:
 
 	// UActorComponent interface
diff --git a/Engine/Plugins/Cameras/GameplayCameras/Source/GameplayCameras/Public/GameFramework/IGameplayCameraSystemHost.h b/Engine/Plugins/Cameras/GameplayCameras/Source/GameplayCameras/Public/GameFramework/IGameplayCameraSystemHost.h
--- a/Engine/Plugins/Cameras/GameplayCameras/Source/GameplayCameras/Public/GameFramework/IGameplayCameraSystemHost.h
+++ b/Engine/Plugins/Cameras/GameplayCameras/Source/GameplayCameras/Public/GameFramework/IGameplayCameraSystemHost.h
@@ -6,6 +6,7 @@
 #include "Core/CameraRigInstanceID.h"
 #include "Core/CameraShakeInstanceID.h"
 #include "GameplayCameras.h"
+#include "Services/CameraActionInstanceID.h"
 #include "Templates/SharedPointerFwd.h"
 #include "UObject/Interface.h"
 
@@ -13,6 +14,7 @@
 
 #define UE_API GAMEPLAYCAMERAS_API
 
+class UCameraAction;
 class UCameraRigAsset;
 class UCameraShakeAsset;
 class UCanvas;
@@ -119,6 +121,18 @@ protected:
 	/** Stops a running camera shake. */
 	UE_API bool StopCameraShake(FCameraShakeInstanceID InInstanceID, bool bImmediately = false);
 
+	/** Starts a new camera action. */
+	FCameraActionInstanceID StartAction(const UCameraAction* CameraAction);
+
+	/** Returns whether a given camera action instance is running. */
+	bool IsActionRunning(const FCameraActionInstanceID InInstanceID) const;
+
+	/** Stops a given camera action instance. */
+	bool StopAction(const FCameraActionInstanceID InInstanceID);
+
+	/** Stops all camera action instances of a given class. */
+	bool StopAllActionsOfClass(TSubclassOf<UCameraAction> InActionClass);
+
 protected:
 
 	/** The camera system evaluator. */
diff --git a/Engine/Plugins/Cameras/GameplayCameras/Source/GameplayCameras/Public/Services/CameraAction.h b/Engine/Plugins/Cameras/GameplayCameras/Source/GameplayCameras/Public/Services/CameraAction.h
new file mode 100644
--- /dev/null
+++ b/Engine/Plugins/Cameras/GameplayCameras/Source/GameplayCameras/Public/Services/CameraAction.h
@@ -0,0 +1,51 @@
+// Copyright Epic Games, Inc. All Rights Reserved.
+
+#pragma once
+
+#include "UObject/Object.h"
+
+#include "CameraAction.generated.h"
+
+#define UE_API GAMEPLAYCAMERAS_API
+
+namespace UE::Cameras
+{
+	class FCameraActionEvaluator;
+	struct FCameraActionEvaluatorBuilder;
+}
+
+using FCameraActionEvaluatorPtr = UE::Cameras::FCameraActionEvaluator*;
+
+/**
+ * The base class for a camera action. A camera action is a potentially long-running operation
+ * that applies to camera rig instances updating in the main layer.
+ */
+UCLASS(MinimalAPI, Abstract, BlueprintType, DefaultToInstanced, EditInlineNew)
+class UCameraAction : public UObject
+{
+	GENERATED_BODY()
+	
+public:
+
+	/**
+	 * Whether this action will be cloned onto new camera rigs as they get activated on the
+	 * main layer.
+	 */
+	UPROPERTY(EditAnywhere, BlueprintReadWrite, Category=Common, meta=(ExposeOnSpawn=true))
+	bool bPropagateToNewCameraRigs = false;
+	
+public:
+
+	using FCameraActionEvaluatorBuilder = UE::Cameras::FCameraActionEvaluatorBuilder;
+
+	/** Builds the evaluator for this action. */
+	UE_API FCameraActionEvaluatorPtr BuildEvaluator(FCameraActionEvaluatorBuilder& Builder) const;
+
+protected:
+
+	/** Builds the evaluator for this action. */
+	virtual FCameraActionEvaluatorPtr OnBuildEvaluator(FCameraActionEvaluatorBuilder& Builder) const { return nullptr; }
+};
+
+#undef UE_API
+
diff --git a/Engine/Plugins/Cameras/GameplayCameras/Source/GameplayCameras/Public/Services/CameraActionEvaluator.h b/Engine/Plugins/Cameras/GameplayCameras/Source/GameplayCameras/Public/Services/CameraActionEvaluator.h
new file mode 100644
--- /dev/null
+++ b/Engine/Plugins/Cameras/GameplayCameras/Source/GameplayCameras/Public/Services/CameraActionEvaluator.h
@@ -0,0 +1,237 @@
+// Copyright Epic Games, Inc. All Rights Reserved.
+
+#pragma once
+
+#include "Core/CameraObjectRtti.h"
+#include "GameplayCameras.h"
+#include "UObject/ObjectPtr.h"
+
+#define UE_API GAMEPLAYCAMERAS_API
+
+class FReferenceCollector;
+class UCameraAction;
+
+namespace UE::Cameras
+{
+
+class FCameraActionEvaluator;
+class FCameraActionScope;
+struct FCameraNodeEvaluationParams;
+struct FCameraNodeEvaluationResult;
+
+#if UE_GAMEPLAY_CAMERAS_DEBUG
+struct FCameraDebugBlockBuilder;
+struct FCameraDebugBlockBuildParams;
+#endif
+
+/** 
+ * Flags for camera action evaluators.
+ */
+enum class ECameraActionEvaluatorFlags
+{
+	None = 0,
+	CloneToNewActionScopes = 1,
+
+	Default = None
+};
+ENUM_CLASS_FLAGS(ECameraActionEvaluatorFlags)
+
+/**
+ * Structure for initializing a new camera action evaluator.
+ */
+struct FCameraActionEvaluatorInitializeParams
+{
+	/** The scope in which the action will be running. */
+	TSharedPtr<FCameraActionScope> Scope;
+};
+
+/**
+ * Structure for tearing down a camera action evaluator.
+ */
+struct FCameraActionEvaluatorTeardownParams
+{
+	/** The scope in which the action was be running. */
+	TSharedPtr<FCameraActionScope> Scope;
+};
+
+/**
+ * Parameter structure for updating a camera action evaluator.
+ */
+struct FCameraActionEvaluationParams
+{
+	FCameraActionEvaluationParams(const FCameraNodeEvaluationParams& InParams)
+		: EvaluationParams(InParams)
+	{}
+
+	/** Information about the evaluation pass that is currently happening. */
+	const FCameraNodeEvaluationParams& EvaluationParams;
+	/** The scope in which the action is running. */
+	TSharedPtr<FCameraActionScope> Scope;
+};
+
+/**
+ * Result of updating the camera action.
+ */
+struct FCameraActionEvaluationResult
+{
+	FCameraActionEvaluationResult(FCameraNodeEvaluationResult& InResult)
+		: Result(InResult)
+	{}
+
+	/** The evaluation result that pass through the camera rig instance. */
+	FCameraNodeEvaluationResult& Result;
+
+	/** Whether this action has ended. */
+	bool bIsActionFinished = false;
+};
+
+/**
+ * Parameter structure for serializing a camera action evaluator.
+ */
+struct FCameraActionEvaluatorSerializeParams
+{
+};
+
+/**
+ * Structure for building new camera actions.
+ */
+struct FCameraActionEvaluatorBuilder
+{
+	FCameraActionEvaluatorBuilder()
+	{}
+
+	template<typename EvaluatorType, typename ...ArgTypes>
+	EvaluatorType* BuildEvaluator(ArgTypes&&... InArgs);
+
+private:
+
+	TSharedPtr<FCameraActionEvaluator> BuiltEvaluator;
+
+	friend class FCameraActionScope;
+};
+
+/**
+ * Base class for a camera action evaluator.
+ */
+class FCameraActionEvaluator
+{
+	UE_GAMEPLAY_CAMERAS_DECLARE_RTTI_BASE(GAMEPLAYCAMERAS_API, FCameraActionEvaluator)
+
+public:
+
+	FCameraActionEvaluator() = default;
+	virtual ~FCameraActionEvaluator() = default;
+
+	/** Initialize this action. */
+	void Initialize(const FCameraActionEvaluatorInitializeParams& Params, FCameraActionEvaluationResult& OutResult);
+
+	/** Tear down this action. */
+	void Teardown(const FCameraActionEvaluatorTeardownParams& Params);
+
+	/** Collect referenced UObjects for this action. */
+	void AddReferencedObjects(FReferenceCollector& Collector);
+
+	/** Run this action before its scope evaluates. */
+	void PreScopeRun(const FCameraActionEvaluationParams& Params, FCameraActionEvaluationResult& OutResult);
+
+	/** Run this action after its scope has evaluated. */
+	void PostScopeRun(const FCameraActionEvaluationParams& Params, FCameraActionEvaluationResult& OutResult);
+
+	/** Serializes the state of this action. */
+	void Serialize(const FCameraActionEvaluatorSerializeParams& Params, FArchive& Ar);
+
+	/** Gets the flags for this evaluator. */
+	ECameraActionEvaluatorFlags  GetActionEvaluatorFlags() const { return PrivateFlags; }
+
+	/** Get the camera action. */
+	const UCameraAction* GetCameraAction() const { return PrivateCameraAction; }
+
+	/** Get the camera action. */
+	template<typename CameraAction>
+	const CameraAction* GetCameraActionAs() const
+	{
+		return Cast<CameraAction>(PrivateCameraAction);
+	}
+
+#if UE_GAMEPLAY_CAMERAS_DEBUG
+	/** Called to create debug blocks for this action evaluator. */
+	void BuildDebugBlocks(const FCameraDebugBlockBuildParams& Params, FCameraDebugBlockBuilder& Builder);
+#endif  // UE_GAMEPLAY_CAMERAS_DEBUG
+
+public:
+
+	// Internal API.
+	void SetPrivateCameraAction(TObjectPtr<const UCameraAction> InCameraAction);
+
+protected:
+
+	/** Initialize this action. */
+	virtual void OnInitialize(const FCameraActionEvaluatorInitializeParams& Params, FCameraActionEvaluationResult& OutResult) {}
+
+	/** Tear down this action. */
+	virtual void OnTeardown(const FCameraActionEvaluatorTeardownParams& Params) {}
+
+	/** Collect referenced UObjects for this action. */
+	virtual void OnAddReferencedObjects(FReferenceCollector& Collector) {}
+
+	/** Run this action before its scope evaluates. */
+	virtual void OnPreScopeRun(const FCameraActionEvaluationParams& Params, FCameraActionEvaluationResult& OutResult) {}
+
+	/** Run this action after its scope has evaluated. */
+	virtual void OnPostScopeRun(const FCameraActionEvaluationParams& Params, FCameraActionEvaluationResult& OutResult) {}
+
+	/** Serializes the state of this evaluator. */
+	virtual void OnSerialize(const FCameraActionEvaluatorSerializeParams& Params, FArchive& Ar) {}
+
+#if UE_GAMEPLAY_CAMERAS_DEBUG
+	/** Called to create debug blocks for this action evaluator. */
+	virtual void OnBuildDebugBlocks(const FCameraDebugBlockBuildParams& Params, FCameraDebugBlockBuilder& Builder) {}
+#endif  // UE_GAMEPLAY_CAMERAS_DEBUG
+
+protected:
+
+	/**
+	 * Sets the flags for this evaluator.
+	 * Can be called from the constructor, or during OnInitialize().
+	 */
+	UE_API void SetActionEvaluatorFlags(ECameraActionEvaluatorFlags InFlags);
+
+	/**
+	 * Adds flags for this evaluator.
+	 * Can be called from the constructor, or during OnInitialize().
+	 */
+	UE_API void AddActionEvaluatorFlags(ECameraActionEvaluatorFlags InFlags);
+
+private:
+
+	/** The camera action to run. */
+	TObjectPtr<const UCameraAction> PrivateCameraAction;
+
+	/** The flags for this evaluator. */
+	ECameraActionEvaluatorFlags PrivateFlags = ECameraActionEvaluatorFlags::Default;
+};
+
+template<typename EvaluatorType, typename ...ArgTypes>
+EvaluatorType* FCameraActionEvaluatorBuilder::BuildEvaluator(ArgTypes&&... InArgs)
+{
+	ensure(!BuiltEvaluator.IsValid());
+	TSharedRef<FCameraActionEvaluator> Evaluator = MakeShared<EvaluatorType>(Forward<ArgTypes>(InArgs)...);
+	BuiltEvaluator = Evaluator;
+	return Evaluator->CastThisChecked<EvaluatorType>();
+}
+
+}  // namespace UE::Cameras
+
+// Utility macros for declaring and defining camera actions.
+//
+#define UE_DECLARE_CAMERA_ACTION_EVALUATOR(ApiDeclSpec, ClassName)\
+	UE_GAMEPLAY_CAMERAS_DECLARE_RTTI(ApiDeclSpec, ClassName, FCameraActionEvaluator)
+
+#define UE_DECLARE_CAMERA_ACTION_EVALUATOR_EX(ApiDeclSpec, ClassName, BaseClassName)\
+	UE_GAMEPLAY_CAMERAS_DECLARE_RTTI(ApiDeclSpec, ClassName, BaseClassName)
+
+#define UE_DEFINE_CAMERA_ACTION_EVALUATOR(ClassName)\
+	UE_GAMEPLAY_CAMERAS_DEFINE_RTTI(ClassName)
+
+#undef UE_API
+
diff --git a/Engine/Plugins/Cameras/GameplayCameras/Source/GameplayCameras/Public/Services/CameraActionInstanceID.h b/Engine/Plugins/Cameras/GameplayCameras/Source/GameplayCameras/Public/Services/CameraActionInstanceID.h
new file mode 100644
--- /dev/null
+++ b/Engine/Plugins/Cameras/GameplayCameras/Source/GameplayCameras/Public/Services/CameraActionInstanceID.h
@@ -0,0 +1,87 @@
+// Copyright Epic Games, Inc. All Rights Reserved.
+
+#pragma once
+
+#include "CoreTypes.h"
+#include "Kismet/BlueprintFunctionLibrary.h"
+
+#include "CameraActionInstanceID.generated.h"
+
+namespace UE::Cameras
+{
+	class FCameraActionScope;
+	class FCameraActionService;
+}
+
+/**
+ * The ID of a camera action instance.
+ */
+USTRUCT(BlueprintType)
+struct FCameraActionInstanceID
+{
+	GENERATED_BODY()
+
+public:
+
+	FCameraActionInstanceID() : Value(INVALID) {}
+
+	bool IsValid() const { return Value != INVALID; }
+
+	explicit operator bool() const { return IsValid(); }
+
+public:
+
+	friend bool operator==(FCameraActionInstanceID A, FCameraActionInstanceID B)
+	{
+		return A.Value == B.Value;
+	}
+
+	friend bool operator!=(FCameraActionInstanceID A, FCameraActionInstanceID B)
+	{
+		return !(A == B);
+	}
+
+	friend uint32 GetTypeHash(FCameraActionInstanceID In)
+	{
+		return In.Value;
+	}
+
+	friend FArchive& operator<< (FArchive& Ar, FCameraActionInstanceID& In)
+	{
+		Ar << In.Value;
+		return Ar;
+	}
+
+private:
+
+	FCameraActionInstanceID(uint32 InValue) 
+		: Value(InValue)
+	{}
+
+	static const uint32 INVALID = uint32(-1);
+
+	UPROPERTY()
+	uint32 Value;
+
+	friend class UE::Cameras::FCameraActionScope;
+	friend class UE::Cameras::FCameraActionService;
+};
+
+/**
+ * Blueprint functions for camera action instance IDs.
+ */
+UCLASS(MinimalAPI)
+class UCameraActionInstanceFunctions : public UBlueprintFunctionLibrary
+{
+	GENERATED_BODY()
+
+public:
+
+	/** 
+	 * Whether the given camera action instance ID is valid.
+	 * A valid ID doesn't necessarily correspond to a camera action instance that is still running.
+	 */
+	UFUNCTION(BlueprintPure, Category="Camera")
+	static bool IsValid(FCameraActionInstanceID InstanceID) { return InstanceID.IsValid(); }
+};
+
diff --git a/Engine/Plugins/Cameras/GameplayCameras/Source/GameplayCameras/Public/Services/CameraActionScope.h b/Engine/Plugins/Cameras/GameplayCameras/Source/GameplayCameras/Public/Services/CameraActionScope.h
new file mode 100644
--- /dev/null
+++ b/Engine/Plugins/Cameras/GameplayCameras/Source/GameplayCameras/Public/Services/CameraActionScope.h
@@ -0,0 +1,126 @@
+// Copyright Epic Games, Inc. All Rights Reserved.
+
+#pragma once
+
+#include "Containers/Array.h"
+#include "Core/CameraRigInstanceID.h"
+#include "GameplayCameras.h"
+#include "Services/CameraActionInstanceID.h"
+#include "Templates/SharedPointer.h"
+
+#define UE_API GAMEPLAYCAMERAS_API
+
+class FReferenceCollector;
+class UCameraAction;
+class UCameraRigAsset;
+
+namespace UE::Cameras
+{
+
+class FCameraActionEvaluator;
+class FCameraEvaluationContext;
+class FCameraNodeEvaluator;
+class FCameraNodeEvaluatorHierarchy;
+struct FCameraActionEvaluationResult;
+struct FCameraActionEvaluatorSerializeParams;
+struct FCameraNodeEvaluationParams;
+struct FCameraNodeEvaluationResult;
+struct FCameraRigEvaluationInfo;
+
+#if UE_GAMEPLAY_CAMERAS_DEBUG
+struct FCameraDebugBlockBuildParams;
+struct FCameraDebugBlockBuilder;
+#endif
+
+/**
+ * A camera action scope, which is a logical container for zero or more camera actions
+ * to be run for a given camera rig instance.
+ */
+class FCameraActionScope : public TSharedFromThis<FCameraActionScope>
+{
+public:
+
+	/** Adds a new camera action. */
+	FCameraActionInstanceID StartAction(const UCameraAction* InCameraAction);
+	
+	/** Returns whether the given action instance is running. */
+	bool IsActionRunning(const FCameraActionInstanceID InInstanceID) const;
+
+	/** Gets a running action. */
+	FCameraActionEvaluator* GetAction(const FCameraActionInstanceID InInstanceID) const;
+
+	/** Stops a running camera action. */
+	bool StopAction(const FCameraActionInstanceID InInstanceID);
+
+public:
+
+	/** Initialize this scope and associate it with the given camera rig. */
+	void Initialize(const FCameraRigEvaluationInfo& CameraRigInfo);
+
+	/** Gets the evaluation context that the camera rig instance runs inside of. */
+	TSharedPtr<const FCameraEvaluationContext> GetEvaluationContext() const { return EvaluationContext.Pin(); }
+
+	/** Gets the camera rig associated with this action scope. */
+	const UCameraRigAsset* GetCameraRig() const { return CameraRig; }
+
+	/** Gets the root evaluator for the camera rig instance that this action acts upon. */
+	FCameraNodeEvaluator* GetChildEvaluator() const { return ChildEvaluator; }
+
+	/** Gets the evaluation info of the camera rig instance that this action acts upon. */
+	FCameraRigEvaluationInfo GetCameraRigEvaluationInfo() const;
+
+public:
+
+	/** Collect referenced UObjects for all actions in the scope. */
+	void AddReferencedObjects(FReferenceCollector& Collector);
+
+	/** Runs the pre-scope evaluation for all actions in the scope. */
+	void PreScopeRun(const FCameraNodeEvaluationParams& Params, FCameraNodeEvaluationResult& OutResult);
+
+	/** Runs the post-scope evaluation for all actions in the scope. */
+	void PostScopeRun(const FCameraNodeEvaluationParams& Params, FCameraNodeEvaluationResult& OutResult);
+
+	/** Serializes the state of all actions in the scope. */
+	void Serialize(const FCameraActionEvaluatorSerializeParams& Params, FArchive& Ar);
+
+#if UE_GAMEPLAY_CAMERAS_DEBUG
+	/** Called to create debug blocks for this action scope. */
+	void BuildDebugBlocks(const FCameraDebugBlockBuildParams& Params, FCameraDebugBlockBuilder& Builder);
+#endif  // UE_GAMEPLAY_CAMERAS_DEBUG
+
+private:
+
+	struct FActionInfo;
+
+	bool PrepareActionForRun(FActionInfo& ActionInfo, FCameraActionEvaluationResult& OutResult);
+
+private:
+
+	enum EActionState
+	{
+		Uninitialized,
+		Running,
+		Finished
+	};
+
+	struct FActionInfo
+	{
+		TSharedPtr<FCameraActionEvaluator> Evaluator;
+		FCameraActionInstanceID InstanceID;
+		EActionState State = EActionState::Uninitialized;
+	};
+
+	FCameraRigInstanceID CameraRigInstanceID;
+	TWeakPtr<const FCameraEvaluationContext> EvaluationContext;
+	TObjectPtr<const UCameraRigAsset> CameraRig;
+	const FCameraNodeEvaluationResult* ChildResult = nullptr;
+	FCameraNodeEvaluator* ChildEvaluator = nullptr;
+
+	TArray<FActionInfo> Actions;
+	uint32 NextActionID = 0;
+};
+
+}  // namespace UE::Cameras
+
+#undef UE_API
+
diff --git a/Engine/Plugins/Cameras/GameplayCameras/Source/GameplayCameras/Public/Services/CameraActionService.h b/Engine/Plugins/Cameras/GameplayCameras/Source/GameplayCameras/Public/Services/CameraActionService.h
new file mode 100644
--- /dev/null
+++ b/Engine/Plugins/Cameras/GameplayCameras/Source/GameplayCameras/Public/Services/CameraActionService.h
@@ -0,0 +1,100 @@
+// Copyright Epic Games, Inc. All Rights Reserved.
+
+#pragma once
+
+#include "Containers/Array.h"
+#include "Core/CameraEvaluationService.h"
+#include "Services/CameraActionInstanceID.h"
+#include "Templates/SharedPointerFwd.h"
+#include "Templates/SubclassOf.h"
+
+#define UE_API GAMEPLAYCAMERAS_API
+
+class UCameraAction;
+class UCameraRigAsset;
+
+namespace UE::Cameras
+{
+
+class FCameraActionScope;
+class FCameraEvaluationContext;
+class FRootCameraNodeEvaluator;
+struct FCameraNodeEvaluationParams;
+struct FCameraNodeEvaluationResult;
+
+/**
+ * An evaluation service for running "camera actions", which allow manipuilating running camera
+ * rig instances, such as making them turn towards a given point, freeze them, and so on.
+ */
+class FCameraActionService : public FCameraEvaluationService
+{
+	UE_DECLARE_CAMERA_EVALUATION_SERVICE(UE_API, FCameraActionService)
+
+public:
+
+	~FCameraActionService();
+
+	/** Starts a new camera action. */
+	FCameraActionInstanceID StartAction(const UCameraAction* CameraAction);
+
+	/** Returns whether a given camera action instance is running. */
+	bool IsActionRunning(const FCameraActionInstanceID InInstanceID) const;
+
+	/** Stops a given camera action instance. */
+	bool StopAction(const FCameraActionInstanceID InInstanceID);
+
+	/** Stops all camera action instances of a given class. */
+	bool StopAllActionsOfClass(TSubclassOf<UCameraAction> InActionClass);
+
+protected:
+
+	// FCameraEvaluationService interface.
+	virtual void OnInitialize(const FCameraEvaluationServiceInitializeParams& Params) override;
+	virtual void OnPostUpdate(const FCameraEvaluationServiceUpdateParams& Params, FCameraEvaluationServiceUpdateResult& OutResult) override;
+#if UE_GAMEPLAY_CAMERAS_DEBUG
+	virtual void OnBuildDebugBlocks(const FCameraDebugBlockBuildParams& Params, FCameraDebugBlockBuilder& Builder) override;
+#endif  // UE_GAMEPLAY_CAMERAS_DEBUG
+
+private:
+
+	struct FActionSetInfo;
+
+	void OnRootCameraNodeCameraRigEvent(const FRootCameraNodeCameraRigEvent& InEvent);
+
+	void AddCloningAction();
+	void RemoveCloningAction();
+	void CloneAction(FActionSetInfo& ActionSet, TSharedRef<FCameraActionScope> NewActionScope);
+
+	void CleanUpActions();
+
+private:
+
+	struct FActionInfo
+	{
+		TWeakPtr<FCameraActionScope> ActionScope = nullptr;
+		FCameraActionInstanceID ActionID;
+	};
+
+	struct FActionSetInfo
+	{
+		TObjectPtr<const UCameraAction> Data;
+		TArray<FActionInfo, TInlineAllocator<4>> Actions;
+		FCameraActionInstanceID InstanceID;
+	};
+
+	FRootCameraNodeEvaluator* RootNodeEvaluator = nullptr;
+
+	TArray<FActionSetInfo> ActionSets;
+
+	uint32 NextActionSetID = 0;
+	uint8 NumCloningActions = 0;
+
+#if UE_GAMEPLAY_CAMERAS_DEBUG
+	friend class FCameraActionServiceDebugBlock;
+#endif
+};
+
+} // namespace UE::Cameras
+
+#undef UE_API
+
