Cameras: when a Boom Arm or Spline Orbit node doesn't have an input node, use a Driven Control Rotation node

Before, these nodes were handling input themselves by querying the Control Rotation themselves. Now they use the seemingly better Driven Control Rotation logic.

Also implement missing serialization of the boom arm length interpolator in the Boom Arm node.

#Synced-CL 48703199

diff --git a/Engine/Plugins/Cameras/GameplayCameras/Source/GameplayCameras/Private/Nodes/Common/BoomArmCameraNode.cpp b/Engine/Plugins/Cameras/GameplayCameras/Source/GameplayCameras/Private/Nodes/Common/BoomArmCameraNode.cpp
--- a/Engine/Plugins/Cameras/GameplayCameras/Source/GameplayCameras/Private/Nodes/Common/BoomArmCameraNode.cpp
+++ b/Engine/Plugins/Cameras/GameplayCameras/Source/GameplayCameras/Private/Nodes/Common/BoomArmCameraNode.cpp
@@ -2,15 +2,12 @@
 
 #include "Nodes/Common/BoomArmCameraNode.h"
 
-#include "Core/CameraEvaluationContext.h"
-#include "Core/CameraOperation.h"
 #include "Core/CameraRigJoints.h"
 #include "Core/CameraValueInterpolator.h"
 #include "Debug/CameraDebugBlock.h"
 #include "Debug/CameraDebugBlockBuilder.h"
 #include "Debug/CameraDebugRenderer.h"
-#include "GameFramework/Pawn.h"
-#include "GameFramework/PlayerController.h"
+#include "Nodes/Input/DrivenControlRotationCameraNode.h"
 #include "Nodes/Input/Input2DCameraNode.h"
 
 #include UE_INLINE_GENERATED_CPP_BY_NAME(BoomArmCameraNode)
@@ -32,7 +29,15 @@ UE_DEFINE_CAMERA_DEBUG_BLOCK_WITH_FIELDS(FBoomArmCameraDebugBlock)
 void FBoomArmCameraNodeEvaluator::OnBuild(const FCameraNodeEvaluatorBuildParams& Params)
 {
 	const UBoomArmCameraNode* BoomArmNode = GetCameraNodeAs<UBoomArmCameraNode>();
-	InputSlotEvaluator = Params.BuildEvaluatorAs<FInput2DCameraNodeEvaluator>(BoomArmNode->InputSlot);
+	if (BoomArmNode->InputSlot)
+	{
+		InputSlotEvaluator = Params.BuildEvaluatorAs<FInput2DCameraNodeEvaluator>(BoomArmNode->InputSlot);
+	}
+	else
+	{
+		const UDrivenControlRotationCameraNode* PlaceholderInput = NewObject<UDrivenControlRotationCameraNode>();
+		InputSlotEvaluator = Params.BuildEvaluatorAs<FInput2DCameraNodeEvaluator>(PlaceholderInput);
+	}
 	if (BoomArmNode->BoomLengthInterpolator)
 	{
 		BoomLengthInterpolator = BoomArmNode->BoomLengthInterpolator->BuildDoubleInterpolator();
@@ -41,7 +46,7 @@ void FBoomArmCameraNodeEvaluator::OnBuild(const FCameraNodeEvaluatorBuildParams&
 
 void FBoomArmCameraNodeEvaluator::OnInitialize(const FCameraNodeEvaluatorInitializeParams& Params, FCameraNodeEvaluationResult& OutResult)
 {
-	SetNodeEvaluatorFlags(ECameraNodeEvaluatorFlags::SupportsOperations);
+	SetNodeEvaluatorFlags(ECameraNodeEvaluatorFlags::NeedsSerialize);
 
 	const UBoomArmCameraNode* BoomArmNode = GetCameraNodeAs<UBoomArmCameraNode>();
 	BoomOffsetReader.Initialize(BoomArmNode->BoomOffset);
@@ -60,17 +65,12 @@ FCameraNodeEvaluatorChildrenView FBoomArmCameraNodeEvaluator::OnGetChildren()
 FRotator FBoomArmCameraNodeEvaluator::ComputeBoomRotation(const FCameraNodeEvaluationParams& Params, FCameraNodeEvaluationResult& OutResult)
 {
 	FRotator3d BoomRotation = FRotator3d::ZeroRotator;
-	if (InputSlotEvaluator)
+	if (ensure(InputSlotEvaluator))
 	{
 		InputSlotEvaluator->Run(Params, OutResult);
 		const FVector2d YawPitch = InputSlotEvaluator->GetInputValue();
 		BoomRotation = FRotator3d(YawPitch.Y, YawPitch.X, 0);
 	}
-	else if (APlayerController* PlayerController = Params.EvaluationContext->GetPlayerController())
-	{
-		const FRotator3d ControlRotation = PlayerController->GetControlRotation();
-		BoomRotation = ControlRotation;
-	}
 	return BoomRotation;
 }
 
@@ -180,23 +180,16 @@ void FBoomArmCameraNodeEvaluator::OnRun(const FCameraNodeEvaluationParams& Param
 	OutResult.CameraRigJoints.AddYawPitchJoint(BoomPivot);
 }
 
-void FBoomArmCameraNodeEvaluator::OnExecuteOperation(const FCameraOperationParams& Params, FCameraOperation& Operation)
+void FBoomArmCameraNodeEvaluator::OnSerialize(const FCameraNodeEvaluatorSerializeParams& Params, FArchive& Ar)
 {
-	if (!InputSlotEvaluator)
+	if (BoomLengthInterpolator)
 	{
-		// If we don't have an input slot, we use the pawn rotation directly in OnRun. So let's handle
-		// some operations by affecting that pawn rotation ourselves.
-		if (FYawPitchCameraOperation* Op = Operation.CastOperation<FYawPitchCameraOperation>())
-		{
-			if (APlayerController* PlayerController = Params.EvaluationContext->GetPlayerController())
-			{
-				FRotator3d ControlRotation = PlayerController->GetControlRotation();
-				ControlRotation.Yaw = Op->Yaw.Apply(ControlRotation.Yaw);
-				ControlRotation.Pitch = Op->Pitch.Apply(ControlRotation.Pitch);
-				PlayerController->SetControlRotation(ControlRotation);
-			}
-		}
+		FCameraValueInterpolatorSerializeParams InterpolatorParams;
+		BoomLengthInterpolator->Serialize(InterpolatorParams, Ar);
 	}
+
+	Ar << LastPivotLocation;
+	Ar << CumulativePull;
 }
 
 #if UE_GAMEPLAY_CAMERAS_DEBUG
diff --git a/Engine/Plugins/Cameras/GameplayCameras/Source/GameplayCameras/Private/Nodes/Common/SplineOrbitCameraNode.cpp b/Engine/Plugins/Cameras/GameplayCameras/Source/GameplayCameras/Private/Nodes/Common/SplineOrbitCameraNode.cpp
--- a/Engine/Plugins/Cameras/GameplayCameras/Source/GameplayCameras/Private/Nodes/Common/SplineOrbitCameraNode.cpp
+++ b/Engine/Plugins/Cameras/GameplayCameras/Source/GameplayCameras/Private/Nodes/Common/SplineOrbitCameraNode.cpp
@@ -2,21 +2,18 @@
 
 #include "Nodes/Common/SplineOrbitCameraNode.h"
 
-#include "Core/CameraEvaluationContext.h"
-#include "Core/CameraOperation.h"
 #include "Core/CameraParameterReader.h"
 #include "Core/CameraRigJoints.h"
 #include "Debug/CameraDebugBlock.h"
 #include "Debug/CameraDebugBlockBuilder.h"
 #include "Debug/CameraDebugRenderer.h"
-#include "GameFramework/Pawn.h"
-#include "GameFramework/PlayerController.h"
 #include "HAL/IConsoleManager.h"
 #include "IGameplayCamerasModule.h"
 #include "IGameplayCamerasLiveEditListener.h"
 #include "IGameplayCamerasLiveEditManager.h"
 #include "Math/CameraNodeSpaceMath.h"
 #include "Math/Ray.h"
+#include "Nodes/Input/DrivenControlRotationCameraNode.h"
 #include "Nodes/Input/Input2DCameraNode.h"
 
 #include UE_INLINE_GENERATED_CPP_BY_NAME(SplineOrbitCameraNode)
@@ -55,7 +52,6 @@ protected:
 	virtual void OnInitialize(const FCameraNodeEvaluatorInitializeParams& Params, FCameraNodeEvaluationResult& OutResult) override;
 	virtual FCameraNodeEvaluatorChildrenView OnGetChildren() override;
 	virtual void OnRun(const FCameraNodeEvaluationParams& Params, FCameraNodeEvaluationResult& OutResult) override;
-	virtual void OnExecuteOperation(const FCameraOperationParams& Params, FCameraOperation& Operation) override;
 
 #if WITH_EDITOR
 	virtual void OnPostEditChangeProperty(const UCameraNode* InCameraNode, const FPropertyChangedEvent& PropertyChangedEvent) override;
@@ -67,8 +63,6 @@ protected:
 
 private:
 
-	APlayerController* GetPlayerController(TSharedPtr<const FCameraEvaluationContext> EvaluationContext) const;
-
 	void RebuildCurves();
 
 private:
@@ -119,12 +113,20 @@ UE_DEFINE_CAMERA_DEBUG_BLOCK_WITH_FIELDS(FSplineOrbitCameraDebugBlock)
 void FSplineOrbitCameraNodeEvaluator::OnBuild(const FCameraNodeEvaluatorBuildParams& Params)
 {
 	const USplineOrbitCameraNode* SplineOrbitNode = GetCameraNodeAs<USplineOrbitCameraNode>();
-	InputSlotEvaluator = Params.BuildEvaluatorAs<FInput2DCameraNodeEvaluator>(SplineOrbitNode->InputSlot);
+	if (SplineOrbitNode->InputSlot)
+	{
+		InputSlotEvaluator = Params.BuildEvaluatorAs<FInput2DCameraNodeEvaluator>(SplineOrbitNode->InputSlot);
+	}
+	else
+	{
+		const UDrivenControlRotationCameraNode* PlaceholderInput = NewObject<UDrivenControlRotationCameraNode>();
+		InputSlotEvaluator = Params.BuildEvaluatorAs<FInput2DCameraNodeEvaluator>(PlaceholderInput);
+	}
 }
 
 void FSplineOrbitCameraNodeEvaluator::OnInitialize(const FCameraNodeEvaluatorInitializeParams& Params, FCameraNodeEvaluationResult& OutResult)
 {
-	SetNodeEvaluatorFlags(ECameraNodeEvaluatorFlags::SupportsOperations);
+	SetNodeEvaluatorFlags(ECameraNodeEvaluatorFlags::None);
 
 	RebuildCurves();
 
@@ -179,17 +181,12 @@ FCameraNodeEvaluatorChildrenView FSplineOrbitCameraNodeEvaluator::OnGetChildren(
 void FSplineOrbitCameraNodeEvaluator::OnRun(const FCameraNodeEvaluationParams& Params, FCameraNodeEvaluationResult& OutResult)
 {
 	FRotator3d OrbitRotation = FRotator3d::ZeroRotator;
-	if (InputSlotEvaluator)
+	if (ensure(InputSlotEvaluator))
 	{
 		InputSlotEvaluator->Run(Params, OutResult);
 		const FVector2d YawPitch = InputSlotEvaluator->GetInputValue();
 		OrbitRotation = FRotator3d(YawPitch.Y, YawPitch.X, 0);
 	}
-	else if (APlayerController* PlayerController = GetPlayerController(Params.EvaluationContext))
-	{
-		const FRotator3d ControlRotation = PlayerController->GetControlRotation();
-		OrbitRotation = ControlRotation;
-	}
 	OrbitRotation.Normalize();
 
 	const USplineOrbitCameraNode* SplineOrbitNode = GetCameraNodeAs<USplineOrbitCameraNode>();
@@ -267,34 +264,6 @@ void FSplineOrbitCameraNodeEvaluator::OnRun(const FCameraNodeEvaluationParams& P
 	OutResult.CameraRigJoints.AddYawPitchJoint(OrbitPivot);
 }
 
-void FSplineOrbitCameraNodeEvaluator::OnExecuteOperation(const FCameraOperationParams& Params, FCameraOperation& Operation)
-{
-	if (!InputSlotEvaluator)
-	{
-		// If we don't have an input slot, we use the pawn rotation directly in OnRun. So let's handle
-		// some operations by affecting that pawn rotation ourselves.
-		if (FYawPitchCameraOperation* Op = Operation.CastOperation<FYawPitchCameraOperation>())
-		{
-			if (APlayerController* PlayerController = GetPlayerController(Params.EvaluationContext))
-			{
-				FRotator3d ControlRotation = PlayerController->GetControlRotation();
-				ControlRotation.Yaw = Op->Yaw.Apply(ControlRotation.Yaw);
-				ControlRotation.Pitch = Op->Pitch.Apply(ControlRotation.Pitch);
-				PlayerController->SetControlRotation(ControlRotation);
-			}
-		}
-	}
-}
-
-APlayerController* FSplineOrbitCameraNodeEvaluator::GetPlayerController(TSharedPtr<const FCameraEvaluationContext> EvaluationContext) const
-{
-	if (EvaluationContext)
-	{
-		return EvaluationContext->GetPlayerController();
-	}
-	return nullptr;
-}
-
 #if WITH_EDITOR
 
 void FSplineOrbitCameraNodeEvaluator::OnPostEditChangeProperty(const UCameraNode* InCameraNode, const FPropertyChangedEvent& PropertyChangedEvent)
diff --git a/Engine/Plugins/Cameras/GameplayCameras/Source/GameplayCameras/Public/Nodes/Common/BoomArmCameraNode.h b/Engine/Plugins/Cameras/GameplayCameras/Source/GameplayCameras/Public/Nodes/Common/BoomArmCameraNode.h
--- a/Engine/Plugins/Cameras/GameplayCameras/Source/GameplayCameras/Public/Nodes/Common/BoomArmCameraNode.h
+++ b/Engine/Plugins/Cameras/GameplayCameras/Source/GameplayCameras/Public/Nodes/Common/BoomArmCameraNode.h
@@ -34,7 +34,7 @@ protected:
 	UE_API virtual void OnInitialize(const FCameraNodeEvaluatorInitializeParams& Params, FCameraNodeEvaluationResult& OutResult) override;
 	UE_API virtual FCameraNodeEvaluatorChildrenView OnGetChildren() override;
 	UE_API virtual void OnRun(const FCameraNodeEvaluationParams& Params, FCameraNodeEvaluationResult& OutResult) override;
-	UE_API virtual void OnExecuteOperation(const FCameraOperationParams& Params, FCameraOperation& Operation) override;
+	UE_API virtual void OnSerialize(const FCameraNodeEvaluatorSerializeParams& Params, FArchive& Ar) override;
 
 #if UE_GAMEPLAY_CAMERAS_DEBUG
 	UE_API virtual void OnBuildDebugBlocks(const FCameraDebugBlockBuildParams& Params, FCameraDebugBlockBuilder& Builder) override;
