Cameras: make value interpolators into constructable objects in Blueprints

#Synced-CL 48703199

diff --git a/Engine/Plugins/Cameras/GameplayCameras/Source/GameplayCameras/Public/Core/CameraValueInterpolator.h b/Engine/Plugins/Cameras/GameplayCameras/Source/GameplayCameras/Public/Core/CameraValueInterpolator.h
--- a/Engine/Plugins/Cameras/GameplayCameras/Source/GameplayCameras/Public/Core/CameraValueInterpolator.h
+++ b/Engine/Plugins/Cameras/GameplayCameras/Source/GameplayCameras/Public/Core/CameraValueInterpolator.h
@@ -219,7 +219,7 @@ protected:
 /**
  * Base class for value interpolator parameters.
  */
-UCLASS(Abstract, DefaultToInstanced, EditInlineNew, CollapseCategories, MinimalAPI)
+UCLASS(MinimalAPI, Abstract, BlueprintType, DefaultToInstanced, EditInlineNew, CollapseCategories)
 class UCameraValueInterpolator : public UObject
 {
 	GENERATED_BODY()
diff --git a/Engine/Plugins/Cameras/GameplayCameras/Source/GameplayCameras/Public/ValueInterpolators/AccelerationDecelerationValueInterpolator.h b/Engine/Plugins/Cameras/GameplayCameras/Source/GameplayCameras/Public/ValueInterpolators/AccelerationDecelerationValueInterpolator.h
--- a/Engine/Plugins/Cameras/GameplayCameras/Source/GameplayCameras/Public/ValueInterpolators/AccelerationDecelerationValueInterpolator.h
+++ b/Engine/Plugins/Cameras/GameplayCameras/Source/GameplayCameras/Public/ValueInterpolators/AccelerationDecelerationValueInterpolator.h
@@ -20,15 +20,15 @@ class UAccelerationDecelerationValueInterpolator : public UCameraValueInterpolat
 public:
 
 	/** The acceleration rate at the start of interpolation. */
-	UPROPERTY(EditAnywhere, Category="Speed")
+	UPROPERTY(EditAnywhere, BlueprintReadWrite, Category="Speed", meta=(ExposeOnSpawn=true))
 	float Acceleration = 1.f;
 
 	/** The maximum speed reachable during interpolation. */
-	UPROPERTY(EditAnywhere, Category="Speed")
+	UPROPERTY(EditAnywhere, BlueprintReadWrite, Category="Speed", meta=(ExposeOnSpawn=true))
 	float MaxSpeed = 10.f;
 
 	/** The deceleration rate at the end of interpolation. */
-	UPROPERTY(EditAnywhere, Category="Speed")
+	UPROPERTY(EditAnywhere, BlueprintReadWrite, Category="Speed", meta=(ExposeOnSpawn=true))
 	float Deceleration = 1.f;
 };
 
diff --git a/Engine/Plugins/Cameras/GameplayCameras/Source/GameplayCameras/Public/ValueInterpolators/CriticalDamperValueInterpolator.h b/Engine/Plugins/Cameras/GameplayCameras/Source/GameplayCameras/Public/ValueInterpolators/CriticalDamperValueInterpolator.h
--- a/Engine/Plugins/Cameras/GameplayCameras/Source/GameplayCameras/Public/ValueInterpolators/CriticalDamperValueInterpolator.h
+++ b/Engine/Plugins/Cameras/GameplayCameras/Source/GameplayCameras/Public/ValueInterpolators/CriticalDamperValueInterpolator.h
@@ -20,7 +20,7 @@ class UCriticalDamperValueInterpolator : public UCameraValueInterpolator
 public:
 
 	/** The damping factor of the spring-mass system. */
-	UPROPERTY(EditAnywhere, Category="Damping")
+	UPROPERTY(EditAnywhere, BlueprintReadWrite, Category="Damping", meta=(ExposeOnSpawn=true))
 	float DampingFactor = 1.f;
 };
 
diff --git a/Engine/Plugins/Cameras/GameplayCameras/Source/GameplayCameras/Public/ValueInterpolators/DoubleIIRValueInterpolator.h b/Engine/Plugins/Cameras/GameplayCameras/Source/GameplayCameras/Public/ValueInterpolators/DoubleIIRValueInterpolator.h
--- a/Engine/Plugins/Cameras/GameplayCameras/Source/GameplayCameras/Public/ValueInterpolators/DoubleIIRValueInterpolator.h
+++ b/Engine/Plugins/Cameras/GameplayCameras/Source/GameplayCameras/Public/ValueInterpolators/DoubleIIRValueInterpolator.h
@@ -19,15 +19,15 @@ class UDoubleIIRValueInterpolator : public UCameraValueInterpolator
 public:
 
 	/** The primary speed of interpolation. */
-	UPROPERTY(EditAnywhere, Category="Speed")
+	UPROPERTY(EditAnywhere, BlueprintReadWrite, Category="Speed", meta=(ExposeOnSpawn=true))
 	float PrimarySpeed = 1.f;
 
 	/** The intermediate speed of interpolation. */
-	UPROPERTY(EditAnywhere, Category="Speed")
+	UPROPERTY(EditAnywhere, BlueprintReadWrite, Category="Speed", meta=(ExposeOnSpawn=true))
 	float IntermediateSpeed = 1.f;
 
 	/** Whether to use fixed-step evaluation. */
-	UPROPERTY(EditAnywhere, Category="Interpolation")
+	UPROPERTY(EditAnywhere, BlueprintReadWrite, Category="Interpolation", meta=(ExposeOnSpawn=true))
 	bool bUseFixedStep = true;
 };
 
diff --git a/Engine/Plugins/Cameras/GameplayCameras/Source/GameplayCameras/Public/ValueInterpolators/IIRValueInterpolator.h b/Engine/Plugins/Cameras/GameplayCameras/Source/GameplayCameras/Public/ValueInterpolators/IIRValueInterpolator.h
--- a/Engine/Plugins/Cameras/GameplayCameras/Source/GameplayCameras/Public/ValueInterpolators/IIRValueInterpolator.h
+++ b/Engine/Plugins/Cameras/GameplayCameras/Source/GameplayCameras/Public/ValueInterpolators/IIRValueInterpolator.h
@@ -21,11 +21,11 @@ class UIIRValueInterpolator : public UCameraValueInterpolator
 public:
 
 	/** The speed of interpolation. */
-	UPROPERTY(EditAnywhere, Category="Interpolation")
+	UPROPERTY(EditAnywhere, BlueprintReadWrite, Category="Interpolation", meta=(ExposeOnSpawn=true))
 	float Speed = 1.f;
 
 	/** Whether to use fixed-step evaluation. */
-	UPROPERTY(EditAnywhere, Category="Interpolation")
+	UPROPERTY(EditAnywhere, BlueprintReadWrite, Category="Interpolation", meta=(ExposeOnSpawn=true))
 	bool bUseFixedStep = true;
 };
 
