Cameras: new "lock user input" operation

This operation locks user input for one frame.

#Synced-CL 48703199

diff --git a/Engine/Plugins/Cameras/GameplayCameras/Source/GameplayCameras/Private/Nodes/Input/CameraRigInput1DSlot.cpp b/Engine/Plugins/Cameras/GameplayCameras/Source/GameplayCameras/Private/Nodes/Input/CameraRigInput1DSlot.cpp
--- a/Engine/Plugins/Cameras/GameplayCameras/Source/GameplayCameras/Private/Nodes/Input/CameraRigInput1DSlot.cpp
+++ b/Engine/Plugins/Cameras/GameplayCameras/Source/GameplayCameras/Private/Nodes/Input/CameraRigInput1DSlot.cpp
@@ -58,38 +58,48 @@ void FCameraRigInput1DSlotEvaluator::OnRun(const FCameraNodeEvaluationParams& Pa
 {
 	const UCameraRigInput1DSlot* SlotNode = GetCameraNodeAs<UCameraRigInput1DSlot>();
 
-	const double Speed = 
-		SlotNode->bIsPreBlended ?
+	if (!bIsLockedThisFrame)
+	{
+		const double Speed = 
+			SlotNode->bIsPreBlended ?
 			OutResult.VariableTable.GetValue<double>(SlotNode->GetSpeedVariableID()) :
 			SpeedReader.Get(OutResult.VariableTable);
 
-	double FinalDelta(DeltaInputValue * Speed * Params.DeltaTime);
+		double FinalDelta(DeltaInputValue * Speed * Params.DeltaTime);
 
-	if (bIsAccumulated)
-	{
-		InputValue += FinalDelta;
-	}
-	else
-	{
-		InputValue = FinalDelta;
-	}
+		if (bIsAccumulated)
+		{
+			InputValue += FinalDelta;
+		}
+		else
+		{
+			InputValue = FinalDelta;
+		}
 
-	InputValue = SlotNode->Normalize.NormalizeValue(InputValue);
-	InputValue = SlotNode->Clamp.ClampValue(InputValue);
+		InputValue = SlotNode->Normalize.NormalizeValue(InputValue);
+
+		InputValue = SlotNode->Clamp.ClampValue(InputValue);
+	}
 
 	OutResult.VariableTable.SetValue<double>(SlotNode->GetVariableID(), InputValue);
+
+	bIsLockedThisFrame = false;	
 }
 
 void FCameraRigInput1DSlotEvaluator::OnExecuteOperation(const FCameraOperationParams& Params, FCameraOperation& Operation)
 {
-	if (FSingleValueCameraOperation* Op = Operation.CastOperation<FSingleValueCameraOperation>())
+	if (FSingleValueCameraOperation* YawPitchOp = Operation.CastOperation<FSingleValueCameraOperation>())
 	{
 		const UCameraRigInput1DSlot* SlotNode = GetCameraNodeAs<UCameraRigInput1DSlot>();
 
 		double MinValue, MaxValue;
 		SlotNode->Clamp.GetEffectiveClamping(MinValue, MaxValue);
 
-		InputValue = Op->Value.Apply(InputValue, MinValue, MaxValue);
+		InputValue = YawPitchOp->Value.Apply(InputValue, MinValue, MaxValue);
+	}
+	else if (FLockUserInputCameraOperation* LockOp = Operation.CastOperation<FLockUserInputCameraOperation>())
+	{
+		bIsLockedThisFrame = LockOp->bLockThisFrame;
 	}
 }
 
diff --git a/Engine/Plugins/Cameras/GameplayCameras/Source/GameplayCameras/Private/Nodes/Input/CameraRigInput2DSlot.cpp b/Engine/Plugins/Cameras/GameplayCameras/Source/GameplayCameras/Private/Nodes/Input/CameraRigInput2DSlot.cpp
--- a/Engine/Plugins/Cameras/GameplayCameras/Source/GameplayCameras/Private/Nodes/Input/CameraRigInput2DSlot.cpp
+++ b/Engine/Plugins/Cameras/GameplayCameras/Source/GameplayCameras/Private/Nodes/Input/CameraRigInput2DSlot.cpp
@@ -59,43 +59,48 @@ void FCameraRigInput2DSlotEvaluator::OnRun(const FCameraNodeEvaluationParams& Pa
 {
 	const UCameraRigInput2DSlot* SlotNode = GetCameraNodeAs<UCameraRigInput2DSlot>();
 
-	const FVector2d Speed = 
-		SlotNode->bIsPreBlended ?
+	if (!bIsLockedThisFrame)
+	{
+		const FVector2d Speed = 
+			SlotNode->bIsPreBlended ?
 			OutResult.VariableTable.GetValue<FVector2d>(SlotNode->GetSpeedVariableID()) :
 			SpeedReader.Get(OutResult.VariableTable);
 
-	FVector2d FinalDelta(DeltaInputValue.X * Speed.X * Params.DeltaTime, DeltaInputValue.Y * Speed.Y * Params.DeltaTime);
-
-	if (RevertAxisXReader.Get(OutResult.VariableTable))
-	{
-		FinalDelta.X = -FinalDelta.X;
-	}
-	if (RevertAxisYReader.Get(OutResult.VariableTable))
-	{
-		FinalDelta.Y = -FinalDelta.Y;
+		FVector2d FinalDelta(DeltaInputValue.X * Speed.X * Params.DeltaTime, DeltaInputValue.Y * Speed.Y * Params.DeltaTime);
+
+		if (RevertAxisXReader.Get(OutResult.VariableTable))
+		{
+			FinalDelta.X = -FinalDelta.X;
+		}
+		if (RevertAxisYReader.Get(OutResult.VariableTable))
+		{
+			FinalDelta.Y = -FinalDelta.Y;
+		}
+
+		if (bIsAccumulated)
+		{
+			InputValue += FinalDelta;
+		}
+		else
+		{
+			InputValue = FinalDelta;
+		}
+
+		InputValue.X = SlotNode->NormalizeX.NormalizeValue(InputValue.X);
+		InputValue.Y = SlotNode->NormalizeY.NormalizeValue(InputValue.Y);
+
+		InputValue.X = SlotNode->ClampX.ClampValue(InputValue.X);
+		InputValue.Y = SlotNode->ClampY.ClampValue(InputValue.Y);
 	}
 
-	if (bIsAccumulated)
-	{
-		InputValue += FinalDelta;
-	}
-	else
-	{
-		InputValue = FinalDelta;
-	}
-
-	InputValue.X = SlotNode->NormalizeX.NormalizeValue(InputValue.X);
-	InputValue.Y = SlotNode->NormalizeY.NormalizeValue(InputValue.Y);
-
-	InputValue.X = SlotNode->ClampX.ClampValue(InputValue.X);
-	InputValue.Y = SlotNode->ClampY.ClampValue(InputValue.Y);
-
 	OutResult.VariableTable.SetValue<FVector2d>(SlotNode->GetVariableID(), InputValue);
+
+	bIsLockedThisFrame = false;
 }
 
 void FCameraRigInput2DSlotEvaluator::OnExecuteOperation(const FCameraOperationParams& Params, FCameraOperation& Operation)
 {
-	if (FYawPitchCameraOperation* Op = Operation.CastOperation<FYawPitchCameraOperation>())
+	if (FYawPitchCameraOperation* YawPitchOp = Operation.CastOperation<FYawPitchCameraOperation>())
 	{
 		const UCameraRigInput2DSlot* SlotNode = GetCameraNodeAs<UCameraRigInput2DSlot>();
 
@@ -105,8 +110,12 @@ void FCameraRigInput2DSlotEvaluator::OnExecuteOperation(const FCameraOperationPa
 		double MinValueY, MaxValueY;
 		SlotNode->ClampY.GetEffectiveClamping(MinValueY, MaxValueY);
 
-		InputValue.X = Op->Yaw.Apply(InputValue.X, MinValueX, MaxValueX);
-		InputValue.Y = Op->Pitch.Apply(InputValue.Y, MinValueY, MaxValueY);
+		InputValue.X = YawPitchOp->Yaw.Apply(InputValue.X, MinValueX, MaxValueX);
+		InputValue.Y = YawPitchOp->Pitch.Apply(InputValue.Y, MinValueY, MaxValueY);
+	}
+	else if (FLockUserInputCameraOperation* LockOp = Operation.CastOperation<FLockUserInputCameraOperation>())
+	{
+		bIsLockedThisFrame = LockOp->bLockThisFrame;
 	}
 }
 
diff --git a/Engine/Plugins/Cameras/GameplayCameras/Source/GameplayCameras/Private/Nodes/Input/DrivenControlRotationCameraNode.cpp b/Engine/Plugins/Cameras/GameplayCameras/Source/GameplayCameras/Private/Nodes/Input/DrivenControlRotationCameraNode.cpp
--- a/Engine/Plugins/Cameras/GameplayCameras/Source/GameplayCameras/Private/Nodes/Input/DrivenControlRotationCameraNode.cpp
+++ b/Engine/Plugins/Cameras/GameplayCameras/Source/GameplayCameras/Private/Nodes/Input/DrivenControlRotationCameraNode.cpp
@@ -181,12 +181,17 @@ protected:
 	virtual void OnBuildDebugBlocks(const FCameraDebugBlockBuildParams& Params, FCameraDebugBlockBuilder& Builder);
 #endif  // UE_GAMEPLAY_CAMERAS_DEBUG
 
+private:
+
+	void LimitControlRotation(APlayerController* PlayerController, FRotator3d& InOutControlRotation);
+
 private:
 
 	TSharedPtr<FDrivenControlRotationHelperService> HelperService;
 
 	FRotator3d LastControlRotation = FRotator3d::ZeroRotator;
 	bool bLastWasActiveCameraRig = false;
+	bool bIsLockedThisFrame = false;
 
 #if UE_GAMEPLAY_CAMERAS_DEBUG
 	FRotator3d LastDeltaControlRotation;
@@ -218,18 +223,16 @@ void FDrivenControlRotationCameraNodeEvaluator::OnTeardown(const FCameraNodeEval
 
 void FDrivenControlRotationCameraNodeEvaluator::OnRun(const FCameraNodeEvaluationParams& Params, FCameraNodeEvaluationResult& OutResult)
 {
-	APlayerController* PlayerController = Params.EvaluationContext->GetPlayerController();
-	if (PlayerController)
+	if (APlayerController* PlayerController = Params.EvaluationContext->GetPlayerController())
 	{
-
 		const FRotator3d ControlRotation = PlayerController->GetControlRotation();
 
-		if (Params.bIsActiveCameraRig)
+		if (Params.bIsActiveCameraRig && !bIsLockedThisFrame)
 		{
 			// We are running in the active camera rig. Just move with the control rotation.
 			InputValue = FVector2d(ControlRotation.Yaw, ControlRotation.Pitch);
 		}
-		else
+		else if (!Params.bIsActiveCameraRig && !bIsLockedThisFrame)
 		{
 			// We are not the active camera rig anymore. Only apply the delta between the last frame's
 			// control rotation and this frame's control rotation. One important thing to note here
@@ -243,19 +246,13 @@ void FDrivenControlRotationCameraNodeEvaluator::OnRun(const FCameraNodeEvaluatio
 			// messed up by new camera rigs' initializations.
 			const FRotator3d CachedControlRotation = HelperService->GetCachedControlRotation(Params.EvaluationContext);
 			FRotator3d DeltaRotation((CachedControlRotation - LastControlRotation).GetNormalized());
-			InputValue.X += DeltaRotation.Yaw;
-			InputValue.Y += DeltaRotation.Pitch;
+			FRotator3d LocalControlRotation(InputValue.Y + DeltaRotation.Pitch, InputValue.X + DeltaRotation.Yaw, 0.0);
 
 			// If we have a player camera manager, apply its limits to the correct angles.
-			if (APlayerCameraManager* CameraManager = PlayerController->PlayerCameraManager)
-			{
-				FRotator3d LocalControlRotation(InputValue.Y, InputValue.X, 0.0);
-				CameraManager->LimitViewPitch(LocalControlRotation, CameraManager->ViewPitchMin, CameraManager->ViewPitchMax);
-				CameraManager->LimitViewYaw(LocalControlRotation, CameraManager->ViewYawMin, CameraManager->ViewYawMax);
-				CameraManager->LimitViewRoll(LocalControlRotation, CameraManager->ViewRollMin, CameraManager->ViewRollMax);
-				InputValue.X = LocalControlRotation.Yaw;
-				InputValue.Y = LocalControlRotation.Pitch;
-			}
+			LimitControlRotation(PlayerController, LocalControlRotation);
+
+			InputValue.X = LocalControlRotation.Yaw;
+			InputValue.Y = LocalControlRotation.Pitch;
 
 #if UE_GAMEPLAY_CAMERAS_DEBUG
 			LastDeltaControlRotation = DeltaRotation;
@@ -266,40 +263,74 @@ void FDrivenControlRotationCameraNodeEvaluator::OnRun(const FCameraNodeEvaluatio
 	}
 
 	bLastWasActiveCameraRig = Params.bIsActiveCameraRig;
+
+	bIsLockedThisFrame = false;
 }
 
 void FDrivenControlRotationCameraNodeEvaluator::OnExecuteOperation(const FCameraOperationParams& Params, FCameraOperation& Operation)
 {
-	if (FYawPitchCameraOperation* Op = Operation.CastOperation<FYawPitchCameraOperation>())
+	if (FYawPitchCameraOperation* YawPitchOp = Operation.CastOperation<FYawPitchCameraOperation>())
 	{
-		// If we are running in the active camera rig, we have permission to affect the control rotation.
-		// Let's apply the yaw/pitch operation to it.
-		APlayerController* PlayerController = Params.EvaluationContext->GetPlayerController();
-		if (bLastWasActiveCameraRig && PlayerController)
+		if (APlayerController* PlayerController = Params.EvaluationContext->GetPlayerController())
 		{
-			FRotator3d ControlRotation = PlayerController->GetControlRotation();
+			// If we are running in the active camera rig, we have permission to affect the control rotation.
+			// In that case, let's apply the yaw/pitch operation to it. Otherwise, let's only apply the correction
+			// to our local cache.
+			if (bLastWasActiveCameraRig)
+			{
+				FRotator3d ControlRotation = PlayerController->GetControlRotation();
 
-			// Make sure the corrected yaw/pitch angles are in [0..360[ and ]-180..180] respectively.
-			ControlRotation.Yaw = FRotator3d::ClampAxis(Op->Yaw.Apply(ControlRotation.Yaw));
-			ControlRotation.Pitch = FRotator3d::NormalizeAxis(Op->Pitch.Apply(ControlRotation.Pitch));
+				// Make sure the corrected yaw/pitch angles are in [0..360[ and ]-180..180] respectively.
+				ControlRotation.Yaw = FRotator3d::ClampAxis(YawPitchOp->Yaw.Apply(ControlRotation.Yaw));
+				ControlRotation.Pitch = FRotator3d::NormalizeAxis(YawPitchOp->Pitch.Apply(ControlRotation.Pitch));
 
-			// If we have a player camera manager, apply its limits to the correct angles.
-			if (APlayerCameraManager* CameraManager = PlayerController->PlayerCameraManager)
+				LimitControlRotation(PlayerController, ControlRotation);
+
+				PlayerController->SetControlRotation(ControlRotation);
+
+				InputValue = FVector2d(ControlRotation.Yaw, ControlRotation.Pitch);
+			}
+			else
 			{
-				CameraManager->LimitViewPitch(ControlRotation, CameraManager->ViewPitchMin, CameraManager->ViewPitchMax);
-				CameraManager->LimitViewYaw(ControlRotation, CameraManager->ViewYawMin, CameraManager->ViewYawMax);
-				CameraManager->LimitViewRoll(ControlRotation, CameraManager->ViewRollMin, CameraManager->ViewRollMax);
+				InputValue.X = YawPitchOp->Yaw.Apply(InputValue.X);
+				InputValue.Y = YawPitchOp->Pitch.Apply(InputValue.Y);
 			}
-
-			PlayerController->SetControlRotation(ControlRotation);
 		}
 	}
+	else if (FLockUserInputCameraOperation* LockOp = Operation.CastOperation<FLockUserInputCameraOperation>())
+	{
+		// We are told to lock input this frame, so let's keep our input value and overwrite the control rotation to 
+		// keep it to what it was before.
+		APlayerController* PlayerController = Params.EvaluationContext->GetPlayerController();
+		if (LockOp->bLockThisFrame && bLastWasActiveCameraRig && PlayerController)
+		{
+			FRotator3d LockedControlRotation(InputValue.Y, InputValue.X, 0);
+			LimitControlRotation(PlayerController, LockedControlRotation);
+			PlayerController->SetControlRotation(LockedControlRotation);
+		}
+
+		bIsLockedThisFrame = LockOp->bLockThisFrame;
+	}
+}
+
+void FDrivenControlRotationCameraNodeEvaluator::LimitControlRotation(APlayerController* PlayerController, FRotator3d& InOutControlRotation)
+{
+	// If we have a player camera manager, apply its limits to the correct angles.
+	if (APlayerCameraManager* CameraManager = PlayerController->PlayerCameraManager)
+	{
+		CameraManager->LimitViewPitch(InOutControlRotation, CameraManager->ViewPitchMin, CameraManager->ViewPitchMax);
+		CameraManager->LimitViewYaw(InOutControlRotation, CameraManager->ViewYawMin, CameraManager->ViewYawMax);
+		CameraManager->LimitViewRoll(InOutControlRotation, CameraManager->ViewRollMin, CameraManager->ViewRollMax);
+	}
 }
 
 void FDrivenControlRotationCameraNodeEvaluator::OnSerialize(const FCameraNodeEvaluatorSerializeParams& Params, FArchive& Ar)
 {
+	Super::OnSerialize(Params, Ar);
+
 	Ar << LastControlRotation;
 	Ar << bLastWasActiveCameraRig;
+	Ar << bIsLockedThisFrame;
 }
 
 #if UE_GAMEPLAY_CAMERAS_DEBUG
diff --git a/Engine/Plugins/Cameras/GameplayCameras/Source/GameplayCameras/Public/Core/CameraOperation.h b/Engine/Plugins/Cameras/GameplayCameras/Source/GameplayCameras/Public/Core/CameraOperation.h
--- a/Engine/Plugins/Cameras/GameplayCameras/Source/GameplayCameras/Public/Core/CameraOperation.h
+++ b/Engine/Plugins/Cameras/GameplayCameras/Source/GameplayCameras/Public/Core/CameraOperation.h
@@ -15,6 +15,7 @@ enum EBuiltInCameraOperationTypes : uint32
 {
 	YawPitch = 0,
 	SingleValue = 1,
+	LockUserInput = 2,
 
 	MAX = 3
 };
@@ -93,7 +94,9 @@ struct FYawPitchCameraOperation : public FCameraOperation
 {
 	UE_DECLARE_BUILT_IN_CAMERA_OPERATION(FYawPitchCameraOperation, EBuiltInCameraOperationTypes::YawPitch)
 
+	/** The yaw value to set or use as modifier. */
 	FConsumableDouble Yaw;
+	/** The pitch value to set or use as modifier. */
 	FConsumableDouble Pitch;
 };
 
@@ -108,6 +111,17 @@ struct FSingleValueCameraOperation : public FCameraOperation
 	FConsumableDouble Value;
 };
 
+/**
+ * A camera operation that locks user input for a frame.
+ */
+struct FLockUserInputCameraOperation : public FCameraOperation
+{
+	UE_DECLARE_BUILT_IN_CAMERA_OPERATION(FLockUserInputCameraOperation, EBuiltInCameraOperationTypes::LockUserInput)
+
+	/** Whether the targeted nodes should lock their input value for this frame. */
+	bool bLockThisFrame = true;
+};
+
 #undef UE_DECLARE_BUILT_IN_CAMERA_OPERATION
 
 }  // namespace UE::Cameras
diff --git a/Engine/Plugins/Cameras/GameplayCameras/Source/GameplayCameras/Public/Nodes/Input/CameraRigInput1DSlot.h b/Engine/Plugins/Cameras/GameplayCameras/Source/GameplayCameras/Public/Nodes/Input/CameraRigInput1DSlot.h
--- a/Engine/Plugins/Cameras/GameplayCameras/Source/GameplayCameras/Public/Nodes/Input/CameraRigInput1DSlot.h
+++ b/Engine/Plugins/Cameras/GameplayCameras/Source/GameplayCameras/Public/Nodes/Input/CameraRigInput1DSlot.h
@@ -112,6 +112,7 @@ protected:
 
 	double DeltaInputValue = 0.f;
 	bool bIsAccumulated = true;
+	bool bIsLockedThisFrame = false;
 };
 
 }  // namespace UE::Cameras
diff --git a/Engine/Plugins/Cameras/GameplayCameras/Source/GameplayCameras/Public/Nodes/Input/CameraRigInput2DSlot.h b/Engine/Plugins/Cameras/GameplayCameras/Source/GameplayCameras/Public/Nodes/Input/CameraRigInput2DSlot.h
--- a/Engine/Plugins/Cameras/GameplayCameras/Source/GameplayCameras/Public/Nodes/Input/CameraRigInput2DSlot.h
+++ b/Engine/Plugins/Cameras/GameplayCameras/Source/GameplayCameras/Public/Nodes/Input/CameraRigInput2DSlot.h
@@ -127,6 +127,7 @@ protected:
 
 	FVector2d DeltaInputValue = FVector2d::ZeroVector;
 	bool bIsAccumulated = true;
+	bool bIsLockedThisFrame = false;
 };
 
 }  // namespace UE::Cameras
