Cameras: new "aim-at" camera action

This action aims the current active camera rig instance towards a given target. It has options for keeping the lock on the target indefinitely until the action is stopped.

The code for IK-aiming and other rotational corrections has been refactored so it can be used in the new action too.

#Synced-CL 48703199

diff --git a/Engine/Plugins/Cameras/GameplayCameras/Source/GameplayCameras/Private/Actions/AimAtCameraAction.cpp b/Engine/Plugins/Cameras/GameplayCameras/Source/GameplayCameras/Private/Actions/AimAtCameraAction.cpp
new file mode 100644
--- /dev/null
+++ b/Engine/Plugins/Cameras/GameplayCameras/Source/GameplayCameras/Private/Actions/AimAtCameraAction.cpp
@@ -0,0 +1,368 @@
+// Copyright Epic Games, Inc. All Rights Reserved.
+
+#include "Actions/AimAtCameraAction.h"
+
+#include "Core/CameraEvaluationContext.h"
+#include "Core/CameraNodeEvaluator.h"
+#include "Core/CameraNodeEvaluatorHierarchy.h"
+#include "Core/CameraOperation.h"
+#include "Core/CameraRigAsset.h"  // IWYU pragma: keep
+#include "Debug/CameraDebugBlock.h"
+#include "Debug/CameraDebugBlockBuilder.h"
+#include "Debug/CameraDebugRenderer.h"
+#include "Math/CameraAimingMath.h"
+#include "Serialization/MemoryReader.h"
+#include "Serialization/MemoryWriter.h"
+#include "Services/CameraActionEvaluator.h"
+#include "Services/CameraActionScope.h"
+
+#include UE_INLINE_GENERATED_CPP_BY_NAME(AimAtCameraAction)
+
+namespace UE::Cameras
+{
+
+class FAimAtCameraActionEvaluator : public FCameraActionEvaluator
+{
+	UE_DECLARE_CAMERA_ACTION_EVALUATOR(, FAimAtCameraActionEvaluator)
+
+protected:
+
+	// FCameraActionEvaluator
+	virtual void OnInitialize(const FCameraActionEvaluatorInitializeParams& Params, FCameraActionEvaluationResult& OutResult) override;
+	virtual void OnPreScopeRun(const FCameraActionEvaluationParams& Params, FCameraActionEvaluationResult& OutResult) override;
+	virtual void OnSerialize(const FCameraActionEvaluatorSerializeParams& Params, FArchive& Ar) override;
+#if UE_GAMEPLAY_CAMERAS_DEBUG
+	virtual void OnBuildDebugBlocks(const FCameraDebugBlockBuildParams& Params, FCameraDebugBlockBuilder& Builder) override;
+#endif  // UE_GAMEPLAY_CAMERAS_DEBUG
+
+private:
+
+	void LockUserInputThisFrame(const FCameraActionEvaluationParams& Params);
+	bool RunPreviewEvaluation(const FCameraActionEvaluationParams& Params, const FCameraNodeEvaluationResult& Result, FRotator3d& OutCorrection);
+	bool ExecuteYawPitchCorrection(const FCameraActionEvaluationParams& Params, const FRotator3d& Correction);
+
+	bool ContinueAimAction(const FCameraActionEvaluationParams& Params, FCameraActionEvaluationResult& OutResult);
+	bool KeepLockOn(const FCameraActionEvaluationParams& Params, FCameraActionEvaluationResult& OutResult);
+
+private:
+
+	TUniquePtr<TCameraValueInterpolator<FVector2d>> Interpolator;
+
+	FCameraNodeEvaluatorHierarchy ChildHierarchy;
+
+	FCameraNodeEvaluationResult ScratchResult;
+	TArray<uint8> EvaluatorSnapshot;
+
+	FVector3d TargetLocation;
+
+	FCameraPose LastCameraPose;
+	FVector3d LastContextLocation;
+	FVector3d LastPivotLocation;
+
+	bool bIsLockedOn = false;
+
+#if UE_GAMEPLAY_CAMERAS_DEBUG
+	float DebugElapsedTime = 0.f;
+	FRotator3d DebugCorrectionLeft;
+	FRotator3d DebugCurrentCorrection;
+#endif  // UE_GAMEPLAY_CAMERAS_DEBUG
+};
+
+UE_DEFINE_CAMERA_ACTION_EVALUATOR(FAimAtCameraActionEvaluator)
+
+void FAimAtCameraActionEvaluator::OnInitialize(const FCameraActionEvaluatorInitializeParams& Params, FCameraActionEvaluationResult& OutResult)
+{
+	const UAimAtCameraAction* ActionData = GetCameraActionAs<UAimAtCameraAction>();
+
+	if (ActionData->Interpolator)
+	{
+		Interpolator = ActionData->Interpolator->BuildVector2dInterpolator();
+	}
+	else
+	{
+		Interpolator = MakeUnique<TPopValueInterpolator<FVector2d>>();
+	}
+
+	const FCameraRigEvaluationInfo CameraRigEvaluationInfo = Params.Scope->GetCameraRigEvaluationInfo();
+	ChildHierarchy.Build(CameraRigEvaluationInfo.RootEvaluator);
+
+	TargetLocation = ActionData->TargetLocation;
+
+#if UE_GAMEPLAY_CAMERAS_DEBUG
+	DebugElapsedTime = 0.f;
+#endif
+}
+
+void FAimAtCameraActionEvaluator::OnPreScopeRun(const FCameraActionEvaluationParams& Params, FCameraActionEvaluationResult& OutResult)
+{
+	if (Params.EvaluationParams.IsStatelessEvaluation())
+	{
+		return;
+	}
+
+	bool bIsActionStillRunning = true;
+	if (bIsLockedOn)
+	{
+		// We are locked on the target, and are supposed to keep it that way.
+		bIsActionStillRunning = KeepLockOn(Params, OutResult);
+	}
+	else
+	{
+		// Aim the camera towards the target before the camera rig is run.
+		bIsActionStillRunning = ContinueAimAction(Params, OutResult);
+	}
+	OutResult.bIsActionFinished = !bIsActionStillRunning;
+
+#if UE_GAMEPLAY_CAMERAS_DEBUG
+	DebugElapsedTime += Params.EvaluationParams.DeltaTime;
+#endif
+}
+
+void FAimAtCameraActionEvaluator::LockUserInputThisFrame(const FCameraActionEvaluationParams& Params)
+{
+	// Execute the operation to turn the camera.
+	FCameraOperationParams OperationParams;
+	OperationParams.Evaluator = Params.EvaluationParams.Evaluator;
+	OperationParams.EvaluationContext = Params.EvaluationParams.EvaluationContext;
+
+	FLockUserInputCameraOperation Operation;
+	Operation.bLockThisFrame = true;
+
+	ChildHierarchy.CallExecuteOperation(OperationParams, Operation);
+}
+
+bool FAimAtCameraActionEvaluator::RunPreviewEvaluation(const FCameraActionEvaluationParams& Params, const FCameraNodeEvaluationResult& Result, FRotator3d& OutCorrection)
+{
+	ScratchResult.OverrideAll(Result, true);
+
+	// Save the initial state of the camera rig.
+	{
+		EvaluatorSnapshot.Reset();
+
+		FCameraNodeEvaluatorSerializeParams SerializeParams;
+		FMemoryWriter Writer(EvaluatorSnapshot);
+		ChildHierarchy.CallSerialize(SerializeParams, Writer);
+	}
+
+	// Run the camera rig.
+	{
+		FCameraNodeEvaluationParams PreviewParams(Params.EvaluationParams);
+		PreviewParams.EvaluationType = ECameraNodeEvaluationType::IK;
+
+		const FCameraRigEvaluationInfo CameraRigEvaluationInfo = Params.Scope->GetCameraRigEvaluationInfo();
+		FCameraNodeEvaluator* ChildEvaluator = CameraRigEvaluationInfo.RootEvaluator;
+		ChildEvaluator->Run(PreviewParams, ScratchResult);
+	}
+
+	// Restore the state of the camera rig.
+	{
+		FCameraNodeEvaluatorSerializeParams SerializeParams;
+		FMemoryReader Reader(EvaluatorSnapshot);
+		ChildHierarchy.CallSerialize(SerializeParams, Reader);
+	}
+
+	// Find the pivot for this frame.
+	const FCameraRigJoint* PreviewPivot = FCameraAimingMath::FindPivotJoint(ScratchResult.CameraRigJoints);
+	if (!PreviewPivot)
+	{
+		// If there's no pivot in the camera, we can't aim it.
+		UE_LOG(LogCameraSystem, Warning, 
+				TEXT("Can't aim camera rig '%s': it has no pivot joint."),
+				*GetNameSafe(Params.Scope->GetCameraRig()));
+		return false;
+	}
+
+	// Compute the total correction need to aim at the target.
+	FRotator3d TotalCorrection;
+	const bool bGotCorrection = FCameraAimingMath::ComputeTwoBonesCorrection(
+			ScratchResult.CameraPose, PreviewPivot->Transform.GetLocation(), TargetLocation, TotalCorrection);
+	if (!bGotCorrection)
+	{
+		UE_LOG(LogCameraSystem, Warning, 
+				TEXT("Can't aim camera rig '%s': we can't solve an IK correction for the desired target."),
+				*GetNameSafe(Params.Scope->GetCameraRig()));
+		return false;
+	}
+
+	OutCorrection = TotalCorrection;
+	return true;
+}
+
+bool FAimAtCameraActionEvaluator::ExecuteYawPitchCorrection(const FCameraActionEvaluationParams& Params, const FRotator3d& Correction)
+{
+	// Execute the operation to turn the camera.
+	FCameraOperationParams OperationParams;
+	OperationParams.Evaluator = Params.EvaluationParams.Evaluator;
+	OperationParams.EvaluationContext = Params.EvaluationParams.EvaluationContext;
+
+	FYawPitchCameraOperation Operation;
+	Operation.Yaw = FConsumableDouble::Delta(Correction.Yaw);
+	Operation.Pitch = FConsumableDouble::Delta(Correction.Pitch);
+
+	ChildHierarchy.CallExecuteOperation(OperationParams, Operation);
+
+	if (Operation.Yaw.HasValue() || Operation.Pitch.HasValue())
+	{
+		UE_LOG(LogCameraSystem, Warning, 
+				TEXT("Aborting aiming of camera rig '%s': not all corrections were consumed by the camera nodes."),
+				*GetNameSafe(Params.Scope->GetCameraRig()));
+		return false;
+	}
+
+	return true;
+}
+
+bool FAimAtCameraActionEvaluator::ContinueAimAction(const FCameraActionEvaluationParams& Params, FCameraActionEvaluationResult& OutResult)
+{
+	const UAimAtCameraAction* ActionData = GetCameraActionAs<UAimAtCameraAction>();
+	if (!ensure(ActionData))
+	{
+		return false;
+	}
+
+	TSharedPtr<const FCameraEvaluationContext> EvaluationContext = Params.EvaluationParams.EvaluationContext;
+	if (!ensure(EvaluationContext))
+	{
+		return false;
+	}
+
+	// Prevent the user from turning the camera this frame, since we want to continue aiming.
+	LockUserInputThisFrame(Params);
+
+	// Run a preview evaluation of the camera rig to see exactly how much we still need to turn to get to the target.
+	FRotator3d TotalCorrection;
+	if (!RunPreviewEvaluation(Params, OutResult.Result, TotalCorrection))
+	{
+		return false;
+	}
+#if UE_GAMEPLAY_CAMERAS_DEBUG
+	DebugCorrectionLeft = TotalCorrection;
+#endif
+
+	// See if we're close enough to the target given our tolerance margin.
+	if (FMath::Abs(TotalCorrection.Yaw) < ActionData->LockOnAngleTolerance && FMath::Abs(TotalCorrection.Pitch) < ActionData->LockOnAngleTolerance)
+	{
+		bIsLockedOn = true;
+		return (ActionData->LockOnPolicy == EAimAtCameraActionLockOnPolicy::KeepLock);
+	}
+
+	// Update the interpolator to know how much we need to turn _this frame_ to get to the target over time.
+	{
+		FCameraValueInterpolationParams InterpolatorParams;
+		InterpolatorParams.DeltaTime = Params.EvaluationParams.DeltaTime;
+		InterpolatorParams.bIsCameraCut = OutResult.Result.bIsCameraCut;
+
+		FCameraValueInterpolationResult InterpolatorResult(OutResult.Result.VariableTable);
+
+		Interpolator->Reset(FVector2d(TotalCorrection.Yaw, TotalCorrection.Pitch), FVector2d::ZeroVector);
+		Interpolator->Run(InterpolatorParams, InterpolatorResult);
+	}
+
+	const FVector2d InterpValue = Interpolator->GetCurrentValue();
+	const FRotator3d RemainingCorrection(InterpValue.Y, InterpValue.X, 0.0);
+	const FRotator3d CurCorrection = (TotalCorrection - RemainingCorrection).GetNormalized();
+#if UE_GAMEPLAY_CAMERAS_DEBUG
+	DebugCurrentCorrection = CurCorrection;
+#endif
+
+	// Execute the operation to turn the camera.
+	if (!ExecuteYawPitchCorrection(Params, CurCorrection))
+	{
+		return false;
+	}
+
+	// If the interpolator has finished, we're done too.
+	if (Interpolator->IsFinished())
+	{
+		bIsLockedOn = true;
+		return (ActionData->LockOnPolicy == EAimAtCameraActionLockOnPolicy::KeepLock);
+	}
+
+	return true;
+}
+
+bool FAimAtCameraActionEvaluator::KeepLockOn(const FCameraActionEvaluationParams& Params, FCameraActionEvaluationResult& OutResult)
+{
+	// Prevent the user from turning the camera this frame, since we want to keep a lock on the target.
+	LockUserInputThisFrame(Params);
+
+	// Run a preview of what the camera will look at this frame, and figure out the correction to keep it locked
+	// on the target.
+	FRotator3d Correction;
+	if (!RunPreviewEvaluation(Params, OutResult.Result, Correction))
+	{
+		return false;
+	}
+
+	// Run the correction on the camera rig.
+	if (!ExecuteYawPitchCorrection(Params, Correction))
+	{
+		return false;
+	}
+
+	return true;
+}
+
+void FAimAtCameraActionEvaluator::OnSerialize(const FCameraActionEvaluatorSerializeParams& Params, FArchive& Ar)
+{
+	if (Interpolator)
+	{
+		FCameraValueInterpolatorSerializeParams InterpolatorParams;
+		Interpolator->Serialize(InterpolatorParams, Ar);
+	}
+
+	Ar << TargetLocation;
+
+	LastCameraPose.SerializeWithFlags(Ar);
+	Ar << LastContextLocation;
+	Ar << LastPivotLocation;
+
+	Ar << bIsLockedOn;
+}
+
+#if UE_GAMEPLAY_CAMERAS_DEBUG
+
+UE_DECLARE_CAMERA_DEBUG_BLOCK_START(, FAimAtCameraActionDebugBlock)
+	UE_DECLARE_CAMERA_DEBUG_BLOCK_FIELD(float, ElapsedTime)
+	UE_DECLARE_CAMERA_DEBUG_BLOCK_FIELD(FRotator3d, CorrectionLeft)
+	UE_DECLARE_CAMERA_DEBUG_BLOCK_FIELD(FRotator3d, CurrentCorrection)
+	UE_DECLARE_CAMERA_DEBUG_BLOCK_FIELD(bool, bIsLockedOn)
+UE_DECLARE_CAMERA_DEBUG_BLOCK_END()
+
+void FAimAtCameraActionEvaluator::OnBuildDebugBlocks(const FCameraDebugBlockBuildParams& Params, FCameraDebugBlockBuilder& Builder)
+{
+	FAimAtCameraActionDebugBlock& DebugBlock = Builder.AttachDebugBlock<FAimAtCameraActionDebugBlock>();
+	DebugBlock.ElapsedTime = DebugElapsedTime;
+	DebugBlock.CorrectionLeft = DebugCorrectionLeft;
+	DebugBlock.CurrentCorrection = DebugCurrentCorrection;
+	DebugBlock.bIsLockedOn = bIsLockedOn;
+}
+
+UE_DEFINE_CAMERA_DEBUG_BLOCK_WITH_FIELDS(FAimAtCameraActionDebugBlock)
+
+void FAimAtCameraActionDebugBlock::OnDebugDraw(const FCameraDebugBlockDrawParams& Params, FCameraDebugRenderer& Renderer)
+{
+	if (bIsLockedOn)
+	{
+		Renderer.AddText(TEXT("locked on target"));
+	}
+	else
+	{
+		Renderer.AddText(
+				TEXT("elapsed time %.1f ; correction yaw=%.1f/%.1f  pitch=%.1f/%.1f"),
+				ElapsedTime,
+				CurrentCorrection.Yaw, CorrectionLeft.Yaw,
+				CurrentCorrection.Pitch, CorrectionLeft.Pitch);
+	}
+}
+
+#endif  // UE_GAMEPLAY_CAMERAS_DEBUG
+
+}  // namespace UE::Cameras
+
+FCameraActionEvaluatorPtr UAimAtCameraAction::OnBuildEvaluator(FCameraActionEvaluatorBuilder& Builder) const
+{
+	using namespace UE::Cameras;
+	return Builder.BuildEvaluator<FAimAtCameraActionEvaluator>();
+}
+
diff --git a/Engine/Plugins/Cameras/GameplayCameras/Source/GameplayCameras/Private/Core/CameraIKAim.cpp b/Engine/Plugins/Cameras/GameplayCameras/Source/GameplayCameras/Private/Core/CameraIKAim.cpp
--- a/Engine/Plugins/Cameras/GameplayCameras/Source/GameplayCameras/Private/Core/CameraIKAim.cpp
+++ b/Engine/Plugins/Cameras/GameplayCameras/Source/GameplayCameras/Private/Core/CameraIKAim.cpp
@@ -16,6 +16,7 @@
 #include "Engine/Engine.h"
 #include "GameplayCameras.h"
 #include "GameplayCamerasSettings.h"
+#include "Math/CameraAimingMath.h"
 #include "Math/Ray.h"
 #include "Serialization/MemoryReader.h"
 #include "Serialization/MemoryWriter.h"
@@ -133,12 +134,7 @@ void FCameraIKAim::DoRunIteration(const FCameraIKAimParams& Params, const FCamer
 
 	// Look at what we have... in 99% of cases, we're dealing with a camera rig that can be boiled down
 	// to 2 bones and one yaw/pitch articulation.
-	const FCameraVariableDefinition& YawPitchDefinition = FBuiltInCameraVariables::Get().YawPitchDefinition;
-	const FCameraRigJoint* FoundItem = ScratchResult.CameraRigJoints.GetJoints().FindByPredicate(
-			[&YawPitchDefinition](const FCameraRigJoint& Item)
-			{
-				return Item.VariableID == YawPitchDefinition.VariableID;
-			});
+	const FCameraRigJoint* FoundItem = FCameraAimingMath::FindPivotJoint(ScratchResult);
 #if UE_GAMEPLAY_CAMERAS_DEBUG
 	IterationDebugInfo.bFoundSolver = (FoundItem != nullptr);
 #endif  // UE_GAMEPLAY_CAMERAS_DEBUG
@@ -155,17 +151,12 @@ void FCameraIKAim::DoRunIteration(const FCameraIKAimParams& Params, const FCamer
 
 void FCameraIKAim::RunRootCameraNode(const FCameraIKAimParams& Params, const FCameraRigEvaluationInfo& CameraRigInfo)
 {
-	const FCameraSystemEvaluationResult& LastResult = Params.Evaluator->GetEvaluatedResult();
-	const FCameraNodeEvaluationResult& ContextInitialResult = CameraRigInfo.EvaluationContext->GetInitialResult();
-
 	// Reset the scratch result the same way the camera system does it at the beginning of each frame.
 	ScratchResult.Reset();
 
 	// Make sure the camera rig will get its private variables, such as rig interface parameter overrides.
 	ScratchResult.VariableTable.OverrideAll(CameraRigInfo.LastResult->VariableTable, true);
 
-	FRootCameraNodeEvaluator* RootEvaluator = Params.Evaluator->GetRootNodeEvaluator();
-
 	// Save the initial state of the camera rig.
 	{
 		EvaluatorSnapshot.Reset();
@@ -185,6 +176,8 @@ void FCameraIKAim::RunRootCameraNode(const FCameraIKAimParams& Params, const FCa
 		SingleParams.EvaluationParams.EvaluationContext = CameraRigInfo.EvaluationContext;
 		SingleParams.EvaluationParams.Evaluator = Params.Evaluator;
 		SingleParams.CameraRigInfo = CameraRigInfo;
+
+		FRootCameraNodeEvaluator* RootEvaluator = Params.Evaluator->GetRootNodeEvaluator();
 		RootEvaluator->RunSingleCameraRig(SingleParams, ScratchResult);
 	}
 
@@ -264,7 +257,7 @@ bool FCameraIKAim::CheckTolerance(const FCameraIKAimParams& Params, const FCamer
 void FCameraIKAim::AimTwoBonesCameraRig(const FCameraIKAimParams& Params, const FCameraRigEvaluationInfo& CameraRigInfo, const FTransform3d& PivotTransform, FAimIterationInfo& IterationInfo)
 {
 	FRotator3d Correction;
-	const bool bGotCorrection = ComputeTwoBonesCorrection(
+	const bool bGotCorrection = FCameraAimingMath::ComputeTwoBonesCorrection(
 			ScratchResult.CameraPose, PivotTransform.GetLocation(), Params.TargetLocation, Correction);
 #if UE_GAMEPLAY_CAMERAS_DEBUG
 	FCameraIKAimIterationDebugInfo& IterationDebugInfo = LastRunDebugInfo.Iterations.Last();
@@ -311,169 +304,6 @@ void FCameraIKAim::AimTwoBonesCameraRig(const FCameraIKAimParams& Params, const
 #endif  // UE_GAMEPLAY_CAMERAS_DEBUG
 }
 
-bool FCameraIKAim::ComputeTwoBonesCorrection(const FCameraPose& CurrentPose, const FVector3d& PivotLocation, const FVector3d& DesiredTarget, FRotator3d& OutCorrection)
-{
-	// This is roughly the situation we are looking at, as seen from above:
-	//
-	//                         T
-	//                         |
-	//                         |    	
-	//             , - ~ D ~ - X
-	//         , '       |    /| ' ,
-	//       ,           |   / |     ,
-	//      ,            |  /  |      ,
-	//     ,             | /   |       ,
-	//     ,             |/    |       ,
-	//     ,             P     |       ,
-	//     ,              .    |       ,
-	//      ,               .  |      ,
-	//       ,                .C     ,
-	//         ,                  , '
-	//           ' - , _ _ _ ,  '
-	//    
-	// ...where:
-	//
-	//    P : the pivot
-	//    C : the current camera position
-	//    T : the current camera target
-	//    D : the desired camera target
-	//
-	// The sphere is centered on P, with its radius determined by D.
-	// The intersection of the camera's current line of sight with this circle is X.
-	// What we want is to turn the camera by A, the angle between PD and PX.
-	//
-	// We first compute the sphere's properties:
-	const FVector3d PivotToDesiredTarget = (DesiredTarget - PivotLocation);
-	const double SphereRadius = PivotToDesiredTarget.Length();
-
-	// Next we compute the intersection between the camera's line of sight and that sphere.
-	const FVector3d CameraLocation = CurrentPose.GetLocation();
-	const FVector3d CameraAim = CurrentPose.GetAimDir();
-
-	double DistanceToX = 0.0;
-	const bool bGotX = RaySphereIntersectExit(CameraLocation, CameraAim, PivotLocation, SphereRadius, DistanceToX);
-	if (!bGotX)
-	{
-		return false;
-	}
-
-	// Finally we compute the angle between PX and PD.
-	const FVector3d X = CameraLocation + CameraAim * DistanceToX;
-	const FVector3d PX = (X - PivotLocation);
-	const FVector3d PD = (DesiredTarget - PivotLocation);
-	// IMPORTANT NOTE: This assumes a vertical pivot axis!
-	const FRotator3d RotPX = PX.ToOrientationRotator();
-	const FRotator3d RotPD = PD.ToOrientationRotator();
-	OutCorrection = RotPD - RotPX;
-	return true;
-}
-
-bool FCameraIKAim::RaySphereIntersectExit(const FRay3d& Ray, const FVector3d& SphereOrigin, double SphereRadius, double& OutRayIntersectDistance)
-{
-	return RaySphereIntersectExit(Ray.Origin, Ray.Direction, SphereOrigin, SphereRadius, OutRayIntersectDistance);
-}
-
-bool FCameraIKAim::RaySphereIntersectExit(const FVector3d& RayStart, const FVector3d& RayDir, const FVector3d& SphereOrigin, double SphereRadius, double& OutRayIntersectDistance)
-{
-	// The equation for a point on the sphere is:
-	//
-	//    (x - Ox)^2 + (y - Oy)^2 + (z - Oz)^2 = R^2
-	//
-	// ...where:
-	//    x,y,z are the coordinates of the point
-	//    Ox,Oy,Oz are the coordinates of the sphere's origin
-	//    R is the radius of the sphere
-	//
-	// The equation for a point on the ray is:
-	//
-	//    x = rx + L*dx
-	//    y = ry + L*dy
-	//    z = rz + L*dz
-	//
-	// ...where:
-	//    rx,ry,rz are the coordinates of the ray's start
-	//    dx,dy,dz are the coordinates of the ray's direction vector
-	//    L is the linear coordinate of the point along the ray
-	//
-	// So we have:
-	//
-	//    (rx + L*dx - Ox)^2 + (ry + L*dy - Oy)^2 + (rz + L*dz - Oz)^2 = R^2
-	//
-	// Decomposing each term, we have:
-	//
-	//    dx*dx*L^2 + (rx - Ox)^2 + 2*L*dx*(rx - Ox)
-	//    dy*dy*L^2 + (ry - Oy)^2 + 2*L*dy*(ry - Oy)
-	//    dz*dz*L^2 + (rz - Oz)^2 + 2*L*dz*(rz - Oz)
-	//
-	// We can group then around L, which is the variable we're solving for:
-	//
-	//    (dx*dx + dy*dy + dz*dz)*L^2 + 2*(dx*Fx + dy*Fy + dz*Fz)*L + (Fx*Fx + Fy*Fy + Fz*Fz) = R^2
-	//
-	// ...where Fx,Fy,Fz are the coordinates of the vector from the sphere origin to 
-	// the ray's start.
-	//
-	// This is an equation of the form:  
-	//
-	//    a*L^2 + b*L + c = 0
-	//
-	// ...where:
-	//
-	//    a = (dx*dx + dy*dy + dz*dz) = d.d
-	//    b = 2*(dx*Fx + dy*Fy + dz*Fz) = 2*d.F
-	//    c = (Fx*Fx + Fy*Fy + Fz*Fz) - R^2 = F.F - R^2
-	//
-	// The discriminant for an equation of this form is:
-	//
-	//    D = b^2 - 4*a*c
-	//
-	// If the discriminant is positive, the solutions for the equations are:
-	//
-	//    L = (-b +- sqrt(D)) / (2*a)
-	//
-	// If the discriminant is zero, there's only one solution (such as when the ray barely "touches"
-	// the sphere). If the discriminant is negative, there's no real solution.
-	//
-	// In practice, we only want the solution that's "in front" of the ray, and is the "farthest",
-	// because we only want the exit point.
-	//
-	const FVector3d F = (RayStart - SphereOrigin);
-	const double dxd = RayDir.SizeSquared();
-	const double FxF = F.SizeSquared();
-	const double dxF = FVector3d::DotProduct(RayDir, F);
-
-	const double a = dxd;
-	const double b = 2.0 * dxF;
-	const double c = FxF - SphereRadius * SphereRadius;
-
-	const double D = b * b - 4.0 * a * c;
-
-	if (D > 0.0)
-	{
-		const double SqrtD = FMath::Sqrt(D);
-		const double two_a = 2.0 * a;
-
-		const double L1 = (-b + SqrtD) / two_a; 
-		const double L2 = (-b - SqrtD) / two_a; 
-
-		OutRayIntersectDistance = FMath::Max(L1, L2);
-		return true;
-	}
-
-	if (D == 0.0)
-	{
-		const double SqrtD = FMath::Sqrt(D);
-		const double two_a = 2.0 * a;
-
-		const double L = (-b / two_a);
-
-		OutRayIntersectDistance = L;
-		return L >= 0.0;
-	}
-
-	// D < 0.0
-	return false;
-}
-
 #if UE_GAMEPLAY_CAMERAS_DEBUG
 
 void FCameraIKAim::GetLastRunDebugInfo(FCameraIKAimDebugInfo& OutDebugInfo) const
diff --git a/Engine/Plugins/Cameras/GameplayCameras/Source/GameplayCameras/Private/Math/CameraAimingMath.cpp b/Engine/Plugins/Cameras/GameplayCameras/Source/GameplayCameras/Private/Math/CameraAimingMath.cpp
new file mode 100644
--- /dev/null
+++ b/Engine/Plugins/Cameras/GameplayCameras/Source/GameplayCameras/Private/Math/CameraAimingMath.cpp
@@ -0,0 +1,219 @@
+// Copyright Epic Games, Inc. All Rights Reserved.
+
+#include "Math/CameraAimingMath.h"
+
+#include "Core/BuiltInCameraVariables.h"
+#include "Core/CameraNodeEvaluator.h"
+#include "Core/CameraRigJoints.h"
+#include "Math/UnrealMathUtility.h"
+
+namespace UE::Cameras
+{
+
+ECameraAimingResult FCameraAimingMath::ComputeCorrection(const FCameraNodeEvaluationResult& Result, const FVector3d& DesiredTarget, FRotator3d& OutCorrection)
+{
+	if (const FCameraRigJoint* FoundItem = FindPivotJoint(Result))
+	{
+		const FCameraRigJoint& PivotJoint(*FoundItem);
+		const bool bGotCorrection = ComputeTwoBonesCorrection(Result.CameraPose, PivotJoint.Transform.GetLocation(), DesiredTarget, OutCorrection);
+		return bGotCorrection ? ECameraAimingResult::Success : ECameraAimingResult::NoSolution;
+	}
+	return ECameraAimingResult::NoSolver;
+}
+
+bool FCameraAimingMath::ComputeTwoBonesCorrection(const FCameraPose& CurrentPose, const FVector3d& PivotLocation, const FVector3d& DesiredTarget, FRotator3d& OutCorrection)
+{
+	// This is roughly the situation we are looking at, as seen from above:
+	//
+	//                         T
+	//                         |
+	//                         |    	
+	//             , - ~ D ~ - X
+	//         , '       |    /| ' ,
+	//       ,           |   / |     ,
+	//      ,            |  /  |      ,
+	//     ,             | /   |       ,
+	//     ,             |/    |       ,
+	//     ,             P     |       ,
+	//     ,              .    |       ,
+	//      ,               .  |      ,
+	//       ,                .C     ,
+	//         ,                  , '
+	//           ' - , _ _ _ ,  '
+	//    
+	// ...where:
+	//
+	//    P : the pivot
+	//    C : the current camera position
+	//    T : the current camera target
+	//    D : the desired camera target
+	//
+	// The sphere is centered on P, with its radius determined by D.
+	// The intersection of the camera's current line of sight with this circle is X.
+	// What we want is to turn the camera by A, the angle between PD and PX.
+	//
+	// We first compute the sphere's properties:
+	const FVector3d PivotToDesiredTarget = (DesiredTarget - PivotLocation);
+	const double SphereRadius = PivotToDesiredTarget.Length();
+
+	// Next we compute the intersection between the camera's line of sight and that sphere.
+	const FVector3d CameraLocation = CurrentPose.GetLocation();
+	const FVector3d CameraAim = CurrentPose.GetAimDir();
+
+	double DistanceToX = 0.0;
+	const bool bGotX = RaySphereIntersectExit(CameraLocation, CameraAim, PivotLocation, SphereRadius, DistanceToX);
+	if (!bGotX)
+	{
+		return false;
+	}
+
+	// Finally we compute the angle between PX and PD.
+	const FVector3d X = CameraLocation + CameraAim * DistanceToX;
+	const FVector3d PX = (X - PivotLocation);
+	const FVector3d PD = (DesiredTarget - PivotLocation);
+	// IMPORTANT NOTE: This assumes a vertical pivot axis!
+	const FRotator3d RotPX = PX.ToOrientationRotator();
+	const FRotator3d RotPD = PD.ToOrientationRotator();
+	OutCorrection = RotPD - RotPX;
+	return true;
+}
+
+const FCameraRigJoint* FCameraAimingMath::FindPivotJoint(const FCameraNodeEvaluationResult& Result)
+{
+	return FindPivotJoint(Result.CameraRigJoints);
+}
+
+const FCameraRigJoint* FCameraAimingMath::FindPivotJoint(const FCameraRigJoints& CameraRigJoints)
+{
+	const FCameraVariableDefinition& YawPitchDefinition = FBuiltInCameraVariables::Get().YawPitchDefinition;
+	const FCameraRigJoint* FoundItem = CameraRigJoints.GetJoints().FindByPredicate(
+			[&YawPitchDefinition](const FCameraRigJoint& Item)
+			{
+				return Item.VariableID == YawPitchDefinition.VariableID;
+			});
+	return FoundItem;
+}
+
+bool FCameraAimingMath::RaySphereIntersectExit(const FRay3d& Ray, const FVector3d& SphereOrigin, double SphereRadius, double& OutRayIntersectDistance)
+{
+	return RaySphereIntersectExit(Ray.Origin, Ray.Direction, SphereOrigin, SphereRadius, OutRayIntersectDistance);
+}
+
+bool FCameraAimingMath::RaySphereIntersectExit(const FVector3d& RayStart, const FVector3d& RayDir, const FVector3d& SphereOrigin, double SphereRadius, double& OutRayIntersectDistance)
+{
+	// The equation for a point on the sphere is:
+	//
+	//    (x - Ox)^2 + (y - Oy)^2 + (z - Oz)^2 = R^2
+	//
+	// ...where:
+	//    x,y,z are the coordinates of the point
+	//    Ox,Oy,Oz are the coordinates of the sphere's origin
+	//    R is the radius of the sphere
+	//
+	// The equation for a point on the ray is:
+	//
+	//    x = rx + L*dx
+	//    y = ry + L*dy
+	//    z = rz + L*dz
+	//
+	// ...where:
+	//    rx,ry,rz are the coordinates of the ray's start
+	//    dx,dy,dz are the coordinates of the ray's direction vector
+	//    L is the linear coordinate of the point along the ray
+	//
+	// So we have:
+	//
+	//    (rx + L*dx - Ox)^2 + (ry + L*dy - Oy)^2 + (rz + L*dz - Oz)^2 = R^2
+	//
+	// Decomposing each term, we have:
+	//
+	//    dx*dx*L^2 + (rx - Ox)^2 + 2*L*dx*(rx - Ox)
+	//    dy*dy*L^2 + (ry - Oy)^2 + 2*L*dy*(ry - Oy)
+	//    dz*dz*L^2 + (rz - Oz)^2 + 2*L*dz*(rz - Oz)
+	//
+	// We can group then around L, which is the variable we're solving for:
+	//
+	//    (dx*dx + dy*dy + dz*dz)*L^2 + 2*(dx*Fx + dy*Fy + dz*Fz)*L + (Fx*Fx + Fy*Fy + Fz*Fz) = R^2
+	//
+	// ...where Fx,Fy,Fz are the coordinates of the vector from the sphere origin to 
+	// the ray's start.
+	//
+	// This is an equation of the form:  
+	//
+	//    a*L^2 + b*L + c = 0
+	//
+	// ...where:
+	//
+	//    a = (dx*dx + dy*dy + dz*dz) = d.d
+	//    b = 2*(dx*Fx + dy*Fy + dz*Fz) = 2*d.F
+	//    c = (Fx*Fx + Fy*Fy + Fz*Fz) - R^2 = F.F - R^2
+	//
+	// The discriminant for an equation of this form is:
+	//
+	//    D = b^2 - 4*a*c
+	//
+	// If the discriminant is positive, the solutions for the equations are:
+	//
+	//    L = (-b +- sqrt(D)) / (2*a)
+	//
+	// If the discriminant is zero, there's only one solution (such as when the ray barely "touches"
+	// the sphere). If the discriminant is negative, there's no real solution.
+	//
+	// In practice, we only want the solution that's "in front" of the ray, and is the "farthest",
+	// because we only want the exit point.
+	//
+	const FVector3d F = (RayStart - SphereOrigin);
+	const double dxd = RayDir.SizeSquared();
+	const double FxF = F.SizeSquared();
+	const double dxF = FVector3d::DotProduct(RayDir, F);
+
+	const double a = dxd;
+	const double b = 2.0 * dxF;
+	const double c = FxF - SphereRadius * SphereRadius;
+
+	const double D = b * b - 4.0 * a * c;
+
+	if (D > 0.0)
+	{
+		const double SqrtD = FMath::Sqrt(D);
+		const double two_a = 2.0 * a;
+
+		const double L1 = (-b + SqrtD) / two_a; 
+		const double L2 = (-b - SqrtD) / two_a; 
+
+		OutRayIntersectDistance = FMath::Max(L1, L2);
+		return true;
+	}
+
+	if (D == 0.0)
+	{
+		const double SqrtD = FMath::Sqrt(D);
+		const double two_a = 2.0 * a;
+
+		const double L = (-b / two_a);
+
+		OutRayIntersectDistance = L;
+		return L >= 0.0;
+	}
+
+	// D < 0.0
+	return false;
+}
+
+double FCameraAimingMath::GetErrorAngle(const FVector3d& AimA, const FVector3d& AimB)
+{
+	const double MultipliedLengths = AimA.Length() * AimB.Length();
+	if (MultipliedLengths > UE_SMALL_NUMBER)
+	{
+		const double OrthLength = AimA.Cross(AimB).Length();
+		const double ErrorAngle = FMath::Asin(OrthLength / MultipliedLengths);
+		return FMath::RadiansToDegrees(ErrorAngle);
+	}
+	else
+	{
+		return 0.0;
+	}
+}
+
+}  // namespace UE::Cameras
+
diff --git a/Engine/Plugins/Cameras/GameplayCameras/Source/GameplayCameras/Public/Actions/AimAtCameraAction.h b/Engine/Plugins/Cameras/GameplayCameras/Source/GameplayCameras/Public/Actions/AimAtCameraAction.h
new file mode 100644
--- /dev/null
+++ b/Engine/Plugins/Cameras/GameplayCameras/Source/GameplayCameras/Public/Actions/AimAtCameraAction.h
@@ -0,0 +1,51 @@
+// Copyright Epic Games, Inc. All Rights Reserved.
+
+#pragma once
+
+#include "Core/CameraValueInterpolator.h"
+#include "Services/CameraAction.h"
+
+#include "AimAtCameraAction.generated.h"
+
+/** Lock-on policy for the aim-at camera action. */
+UENUM()
+enum class EAimAtCameraActionLockOnPolicy : uint8
+{
+	/** End the action once the aiming is accomplished. */
+	Disengage,
+	/** Continuously aim at the target until the action is explicitly stopped. */
+	KeepLock
+};
+
+/**
+ * A camera action that aims at a given target, using a given interpolation.
+ */
+UCLASS()
+class UAimAtCameraAction : public UCameraAction
+{
+	GENERATED_BODY()
+
+public:
+
+	/** The target to aim at. */
+	UPROPERTY(EditAnywhere, BlueprintReadWrite, Category=Common, meta=(ExposeOnSpawn=true))
+	FVector TargetLocation;
+
+	/** The interpolation to use for aiming. */
+	UPROPERTY(EditAnywhere, BlueprintReadWrite, Category=Common, meta=(ExposeOnSpawn=true))
+	TObjectPtr<UCameraValueInterpolator> Interpolator;
+
+	/** The tolerance within which we can consider the aiming to be locked on target. */
+	UPROPERTY(EditAnywhere, BlueprintReadWrite, Category=Common, meta=(ExposeOnSpawn=true))
+	float LockOnAngleTolerance = 0.01f;
+
+	/** The lock-on policy. */
+	UPROPERTY(EditAnywhere, BlueprintReadWrite, Category=Common, meta=(ExposeOnSpawn=true))
+	EAimAtCameraActionLockOnPolicy LockOnPolicy = EAimAtCameraActionLockOnPolicy::Disengage;
+
+protected:
+
+	// UCameraAction interface.
+	virtual FCameraActionEvaluatorPtr OnBuildEvaluator(FCameraActionEvaluatorBuilder& Builder) const override;
+};
+
diff --git a/Engine/Plugins/Cameras/GameplayCameras/Source/GameplayCameras/Public/Core/CameraIKAim.h b/Engine/Plugins/Cameras/GameplayCameras/Source/GameplayCameras/Public/Core/CameraIKAim.h
--- a/Engine/Plugins/Cameras/GameplayCameras/Source/GameplayCameras/Public/Core/CameraIKAim.h
+++ b/Engine/Plugins/Cameras/GameplayCameras/Source/GameplayCameras/Public/Core/CameraIKAim.h
@@ -113,11 +113,6 @@ private:
 	void RunRootCameraNode(const FCameraIKAimParams& Params, const FCameraRigEvaluationInfo& CameraRigInfo);
 	bool CheckTolerance(const FCameraIKAimParams& Params, const FCameraRigEvaluationInfo& CameraRigInfo, FAimIterationInfo& IterationInfo);
 
-	static bool ComputeTwoBonesCorrection(const FCameraPose& CurrentPose, const FVector3d& PivotLocation, const FVector3d& DesiredTarget, FRotator3d& OutCorrection);
-
-	static bool RaySphereIntersectExit(const FRay3d& Ray, const FVector3d& SphereOrigin, double SphereRadius, double& OutRayIntersectDistance);
-	static bool RaySphereIntersectExit(const FVector3d& RayStart, const FVector3d& RayDir, const FVector3d& SphereOrigin, double SphereRadius, double& OutRayIntersectDistance);
-
 private:
 
 	FCameraNodeEvaluationResult ScratchResult;
diff --git a/Engine/Plugins/Cameras/GameplayCameras/Source/GameplayCameras/Public/Core/CameraNodeEvaluatorHierarchy.h b/Engine/Plugins/Cameras/GameplayCameras/Source/GameplayCameras/Public/Core/CameraNodeEvaluatorHierarchy.h
--- a/Engine/Plugins/Cameras/GameplayCameras/Source/GameplayCameras/Public/Core/CameraNodeEvaluatorHierarchy.h
+++ b/Engine/Plugins/Cameras/GameplayCameras/Source/GameplayCameras/Public/Core/CameraNodeEvaluatorHierarchy.h
@@ -45,7 +45,7 @@ public:
 
 	/** Executes the given predicate on each evaluator in depth-first order. */
 	template<typename PredicateClass>
-	void ForEachEvaluator(PredicateClass&& Predicate)
+	void ForEachEvaluator(PredicateClass&& Predicate) const
 	{
 		for (FCameraNodeEvaluator* Evaluator : FlattenedHierarchy)
 		{
@@ -55,7 +55,7 @@ public:
 
 	/** Executes the given predicate on each evaluator in the specified range in depth-first order. */
 	template<typename PredicateClass>
-	void ForEachEvaluator(const FName TaggedRangeName, PredicateClass&& Predicate)
+	void ForEachEvaluator(const FName TaggedRangeName, PredicateClass&& Predicate) const
 	{
 		const FTaggedRange& TaggedRange = TaggedRanges.FindChecked(TaggedRangeName);
 		for (int32 Index = TaggedRange.StartIndex; Index < TaggedRange.EndIndex; ++Index)
@@ -70,7 +70,7 @@ public:
 	 * depth-first order.
 	 */
 	template<typename PredicateClass>
-	void ForEachEvaluator(ECameraNodeEvaluatorFlags FilterFlags, PredicateClass&& Predicate)
+	void ForEachEvaluator(ECameraNodeEvaluatorFlags FilterFlags, PredicateClass&& Predicate) const
 	{
 		for (FCameraNodeEvaluator* Evaluator : FlattenedHierarchy)
 		{
@@ -86,7 +86,7 @@ public:
 	 * evaluator flags in depth-first order.
 	 */
 	template<typename PredicateClass>
-	void ForEachEvaluator(const FName TaggedRangeName, ECameraNodeEvaluatorFlags FilterFlags, PredicateClass&& Predicate)
+	void ForEachEvaluator(const FName TaggedRangeName, ECameraNodeEvaluatorFlags FilterFlags, PredicateClass&& Predicate) const
 	{
 		const FTaggedRange& TaggedRange = TaggedRanges.FindChecked(TaggedRangeName);
 		for (int32 Index = TaggedRange.StartIndex; Index < TaggedRange.EndIndex; ++Index)
diff --git a/Engine/Plugins/Cameras/GameplayCameras/Source/GameplayCameras/Public/Math/CameraAimingMath.h b/Engine/Plugins/Cameras/GameplayCameras/Source/GameplayCameras/Public/Math/CameraAimingMath.h
new file mode 100644
--- /dev/null
+++ b/Engine/Plugins/Cameras/GameplayCameras/Source/GameplayCameras/Public/Math/CameraAimingMath.h
@@ -0,0 +1,56 @@
+// Copyright Epic Games, Inc. All Rights Reserved.
+
+#pragma once
+
+#include "Math/MathFwd.h"
+
+struct FCameraPose;
+
+namespace UE::Cameras
+{
+
+class FCameraRigJoints;
+struct FCameraNodeEvaluationResult;
+struct FCameraRigJoint;
+
+/** Result for an aiming operation. */
+enum class ECameraAimingResult
+{
+	/** The aiming was successful. */
+	Success,
+	/** The aiming did not have a valid solution. */
+	NoSolution,
+	/** The aiming was not possible for the given camera rig. */
+	NoSolver
+};
+
+/**
+ * A utility class for math operations pertaining to aiming a camera at something.
+ */
+class FCameraAimingMath
+{
+public:
+
+	/** Computes the rotation correction needed for the given camera to be aimed at the given target. */
+	static ECameraAimingResult ComputeCorrection(const FCameraNodeEvaluationResult& Result, const FVector3d& DesiredTarget, FRotator3d& OutCorrection);
+
+	/** Finds the pivot joint in the given evaluation result. */
+	static const FCameraRigJoint* FindPivotJoint(const FCameraNodeEvaluationResult& Result);
+	/** Finds the pivot joint in the given camera rig joints. */
+	static const FCameraRigJoint* FindPivotJoint(const FCameraRigJoints& CameraRigJoints);
+
+	/** Computes the rotation correction needed for the given camera to be aimed at the given target.*/
+	static bool ComputeTwoBonesCorrection(const FCameraPose& CurrentPose, const FVector3d& PivotLocation, const FVector3d& DesiredTarget, FRotator3d& OutCorrection);
+
+	/** Computes the difference between two aim vectors. */
+	static double GetErrorAngle(const FVector3d& AimA, const FVector3d& AimB);
+
+private:
+
+	static bool RaySphereIntersectExit(const FRay3d& Ray, const FVector3d& SphereOrigin, double SphereRadius, double& OutRayIntersectDistance);
+	static bool RaySphereIntersectExit(const FVector3d& RayStart, const FVector3d& RayDir, const FVector3d& SphereOrigin, double SphereRadius, double& OutRayIntersectDistance);
+
+};
+
+}  // namespace UE::Cameras
+
