Cameras: make latest changes backwards compatible to 5.7

#Synced-CL 48703199

diff --git a/Engine/Plugins/Cameras/GameplayCameras/Source/GameplayCameras/Private/Directors/StateTreeCameraDirector.cpp b/Engine/Plugins/Cameras/GameplayCameras/Source/GameplayCameras/Private/Directors/StateTreeCameraDirector.cpp
--- a/Engine/Plugins/Cameras/GameplayCameras/Source/GameplayCameras/Private/Directors/StateTreeCameraDirector.cpp
+++ b/Engine/Plugins/Cameras/GameplayCameras/Source/GameplayCameras/Private/Directors/StateTreeCameraDirector.cpp
@@ -12,6 +12,7 @@
 #include "GameplayCamerasSettings.h"
 #include "Helpers/OutgoingReferenceFinder.h"
 #include "Logging/TokenizedMessage.h"
+#include "Misc/EngineVersionComparison.h"
 #include "StateTree.h"
 #include "StateTreeExecutionContext.h"
 #include "StateTreeInstanceData.h"
@@ -94,7 +95,11 @@ void FStateTreeCameraDirectorEvaluator::OnActivate(const FCameraDirectorActivate
 		return;
 	}
 
+#if UE_VERSION_NEWER_THAN_OR_EQUAL(5,8,0)
 	StateTreeContext.Start(StateTreeReference.GetGlobalParameters());
+#else
+	StateTreeContext.Start(&StateTreeReference.GetParameters());
+#endif
 }
 
 void FStateTreeCameraDirectorEvaluator::OnDeactivate(const FCameraDirectorDeactivateParams& Params, FCameraDirectorEvaluationResult& OutResult)
diff --git a/Engine/Plugins/Cameras/GameplayCameras/Source/GameplayCameras/Private/Services/CameraActionService.cpp b/Engine/Plugins/Cameras/GameplayCameras/Source/GameplayCameras/Private/Services/CameraActionService.cpp
--- a/Engine/Plugins/Cameras/GameplayCameras/Source/GameplayCameras/Private/Services/CameraActionService.cpp
+++ b/Engine/Plugins/Cameras/GameplayCameras/Source/GameplayCameras/Private/Services/CameraActionService.cpp
@@ -10,6 +10,7 @@
 #include "Debug/CameraDebugBlock.h"
 #include "Debug/CameraDebugBlockBuilder.h"
 #include "Debug/CameraDebugRenderer.h"
+#include "Misc/EngineVersionComparison.h"
 #include "Serialization/MemoryReader.h"
 #include "Serialization/MemoryWriter.h"
 #include "Services/CameraAction.h"
@@ -240,12 +241,20 @@ void FCameraActionService::CleanUpActions()
 			{
 				if (!ActionScope->IsActionRunning(ActionInfo.ActionID))
 				{
+#if UE_VERSION_NEWER_THAN_OR_EQUAL(5,8,0)
 					ActionIt.RemoveCurrent(EAllowShrinking::No);
+#else
+					ActionIt.RemoveCurrent();
+#endif
 				}
 			}
 			else
 			{
+#if UE_VERSION_NEWER_THAN_OR_EQUAL(5,8,0)
 				ActionIt.RemoveCurrent(EAllowShrinking::No);
+#else
+				ActionIt.RemoveCurrent();
+#endif
 			}
 		}
 		if (ActionSet.Actions.IsEmpty())
@@ -254,7 +263,11 @@ void FCameraActionService::CleanUpActions()
 			{
 				RemoveCloningAction();
 			}
+#if UE_VERSION_NEWER_THAN_OR_EQUAL(5,8,0)
 			SetIt.RemoveCurrent(EAllowShrinking::No);
+#else
+			SetIt.RemoveCurrent();
+#endif
 		}
 	}
 }
