Cameras: add ways to parameterize post-process settings

#Synced-CL 48703199

diff --git a/Engine/Plugins/Cameras/GameplayCameras/Source/GameplayCameras/Private/GameFramework/BlueprintCameraEvaluationDataRef.cpp b/Engine/Plugins/Cameras/GameplayCameras/Source/GameplayCameras/Private/GameFramework/BlueprintCameraEvaluationDataRef.cpp
--- a/Engine/Plugins/Cameras/GameplayCameras/Source/GameplayCameras/Private/GameFramework/BlueprintCameraEvaluationDataRef.cpp
+++ b/Engine/Plugins/Cameras/GameplayCameras/Source/GameplayCameras/Private/GameFramework/BlueprintCameraEvaluationDataRef.cpp
@@ -208,8 +208,7 @@ void UBlueprintCameraVariableTableFunctionLibrary::SetTransformCameraVariable(co
 	{\
 		FFrame::KismetExecutionMessage(TEXT("No camera context data table has been set"), ELogVerbosity::Error);\
 		return ErrorResult;\
-	}\
-
+	}
 
 FName UBlueprintCameraContextDataTableFunctionLibrary::GetNameData(const FBlueprintCameraEvaluationDataRef& CameraData, FCameraContextDataID DataID)
 {
@@ -303,3 +302,37 @@ bool UBlueprintCameraContextDataTableFunctionLibrary::SetClassData(const FBluepr
 
 #undef UE_PRIVATE_BLUEPRINT_CAMERA_CONTEXT_DATA_TABLE_VALIDATE
 
+#define UE_PRIVATE_BLUEPRINT_CAMERA_CONTEXT_POST_PROCESS_VALIDATE()\
+	using namespace UE::Cameras;\
+	if (!CameraData.IsValid())\
+	{\
+		FFrame::KismetExecutionMessage(TEXT("No valid camera data was provided"), ELogVerbosity::Error);\
+		return;\
+	}\
+
+void UBlueprintPostProcessSettingsCollectionFunctionLibrary::GetPostProcessSettings(const FBlueprintCameraEvaluationDataRef& CameraData, FPostProcessSettings& OutPostProcessSettings)
+{
+	UE_PRIVATE_BLUEPRINT_CAMERA_CONTEXT_POST_PROCESS_VALIDATE()
+	OutPostProcessSettings = CameraData.GetResult()->PostProcessSettings.Get();
+}
+
+void UBlueprintPostProcessSettingsCollectionFunctionLibrary::SetPostProcessSettings(const FBlueprintCameraEvaluationDataRef& CameraData, const FPostProcessSettings& PostProcessSettings)
+{
+	UE_PRIVATE_BLUEPRINT_CAMERA_CONTEXT_POST_PROCESS_VALIDATE()
+	CameraData.GetResult()->PostProcessSettings.Get() = PostProcessSettings;
+}
+
+void UBlueprintPostProcessSettingsCollectionFunctionLibrary::OverrideChangedPostProcessSettings(const FBlueprintCameraEvaluationDataRef& CameraData, const FPostProcessSettings& PostProcessSettings)
+{
+	UE_PRIVATE_BLUEPRINT_CAMERA_CONTEXT_POST_PROCESS_VALIDATE()
+	CameraData.GetResult()->PostProcessSettings.OverrideChanged(PostProcessSettings);
+}
+
+void UBlueprintPostProcessSettingsCollectionFunctionLibrary::LerpPostProcessSettings(const FBlueprintCameraEvaluationDataRef& CameraData, const FPostProcessSettings& PostProcessSettings, float BlendFactor)
+{
+	UE_PRIVATE_BLUEPRINT_CAMERA_CONTEXT_POST_PROCESS_VALIDATE()
+	CameraData.GetResult()->PostProcessSettings.LerpAll(PostProcessSettings, BlendFactor);
+}
+
+#undef UE_PRIVATE_BLUEPRINT_CAMERA_CONTEXT_POST_PROCESS_VALIDATE
+
diff --git a/Engine/Plugins/Cameras/GameplayCameras/Source/GameplayCameras/Private/Nodes/Common/PostProcessCameraNode.cpp b/Engine/Plugins/Cameras/GameplayCameras/Source/GameplayCameras/Private/Nodes/Common/PostProcessCameraNode.cpp
--- a/Engine/Plugins/Cameras/GameplayCameras/Source/GameplayCameras/Private/Nodes/Common/PostProcessCameraNode.cpp
+++ b/Engine/Plugins/Cameras/GameplayCameras/Source/GameplayCameras/Private/Nodes/Common/PostProcessCameraNode.cpp
@@ -2,6 +2,7 @@
 
 #include "Nodes/Common/PostProcessCameraNode.h"
 
+#include "Core/CameraContextDataReader.h"
 #include "Core/CameraPose.h"
 
 #include UE_INLINE_GENERATED_CPP_BY_NAME(PostProcessCameraNode)
@@ -23,17 +24,28 @@ public:
 protected:
 
 	// FCameraNodeEvaluator interface.
+	virtual void OnInitialize(const FCameraNodeEvaluatorInitializeParams& Params, FCameraNodeEvaluationResult& OutResult) override;
 	virtual void OnRun(const FCameraNodeEvaluationParams& Params, FCameraNodeEvaluationResult& OutResult) override;
+
+private:
+
+	TCameraContextDataReader<FPostProcessSettings> PostProcessSettingsReader;
 };
 
 UE_DEFINE_CAMERA_NODE_EVALUATOR(FPostProcessCameraNodeEvaluator)
 
+void FPostProcessCameraNodeEvaluator::OnInitialize(const FCameraNodeEvaluatorInitializeParams& Params, FCameraNodeEvaluationResult& OutResult)
+{
+	const UPostProcessCameraNode* PostProcessNode = GetCameraNodeAs<UPostProcessCameraNode>();
+	PostProcessSettingsReader.Initialize(&PostProcessNode->PostProcessSettings, PostProcessNode->PostProcessSettingsDataID);
+}
+
 void FPostProcessCameraNodeEvaluator::OnRun(const FCameraNodeEvaluationParams& Params, FCameraNodeEvaluationResult& OutResult)
 {
 	FCameraPose& OutPose = OutResult.CameraPose;
 
-	const UPostProcessCameraNode* PostProcessNode = GetCameraNodeAs<UPostProcessCameraNode>();
-	OutResult.PostProcessSettings.OverrideChanged(PostProcessNode->PostProcessSettings);
+	const FPostProcessSettings& PostProcessSettings = PostProcessSettingsReader.GetRef(OutResult.ContextDataTable);
+	OutResult.PostProcessSettings.OverrideChanged(PostProcessSettings);
 }
 
 }  // namespace UE::Cameras
diff --git a/Engine/Plugins/Cameras/GameplayCameras/Source/GameplayCameras/Public/Core/CameraContextDataReader.h b/Engine/Plugins/Cameras/GameplayCameras/Source/GameplayCameras/Public/Core/CameraContextDataReader.h
--- a/Engine/Plugins/Cameras/GameplayCameras/Source/GameplayCameras/Public/Core/CameraContextDataReader.h
+++ b/Engine/Plugins/Cameras/GameplayCameras/Source/GameplayCameras/Public/Core/CameraContextDataReader.h
@@ -4,6 +4,7 @@
 
 #include "Core/CameraContextDataTable.h"
 #include "CoreTypes.h"
+#include <type_traits>
 
 namespace UE::Cameras
 {
@@ -74,6 +75,24 @@ struct TCameraContextDataReader
 		return (DefaultValuePtr ? *DefaultValuePtr : DataType());
 	}
 
+	const DataType& GetRef(const FCameraContextDataTable& ContextDataTable) const
+	{
+		static const DataType StaticDefaultValue = DataType();
+
+		static_assert(!std::is_enum_v<DataType>, "Can't use GetRef with enum types, please use Get instead");
+
+		if (!DataID.IsValid())
+		{
+			return (DefaultValuePtr ? *DefaultValuePtr : StaticDefaultValue);
+		}
+
+		if (const DataType* ActualValue = ContextDataTable.TryGetData<DataType>(DataID))
+		{
+			return *ActualValue;
+		}
+		return (DefaultValuePtr ? *DefaultValuePtr : StaticDefaultValue);
+	}
+
 private:
 
 	const DataType* DefaultValuePtr = nullptr;
diff --git a/Engine/Plugins/Cameras/GameplayCameras/Source/GameplayCameras/Public/GameFramework/BlueprintCameraEvaluationDataRef.h b/Engine/Plugins/Cameras/GameplayCameras/Source/GameplayCameras/Public/GameFramework/BlueprintCameraEvaluationDataRef.h
--- a/Engine/Plugins/Cameras/GameplayCameras/Source/GameplayCameras/Public/GameFramework/BlueprintCameraEvaluationDataRef.h
+++ b/Engine/Plugins/Cameras/GameplayCameras/Source/GameplayCameras/Public/GameFramework/BlueprintCameraEvaluationDataRef.h
@@ -234,4 +234,31 @@ public:
 	static bool SetClassData(const FBlueprintCameraEvaluationDataRef& CameraData, FCameraContextDataID DataID, UClass* Data);
 };
 
+/**
+ * Utility Blueprint functions for camera post process settings.
+ */
+UCLASS()
+class UBlueprintPostProcessSettingsCollectionFunctionLibrary : public UBlueprintFunctionLibrary
+{
+	GENERATED_BODY()
+
+public:
+
+	/** Gets the post-process settings. */
+	UFUNCTION(BlueprintPure, Category=Camera)
+	static void GetPostProcessSettings(const FBlueprintCameraEvaluationDataRef& CameraData, FPostProcessSettings& OutPostProcessSettings);
+
+	/** Sets the post-process settings. */
+	UFUNCTION(BlueprintCallable, Category=Camera)
+	static void SetPostProcessSettings(const FBlueprintCameraEvaluationDataRef& CameraData, const FPostProcessSettings& PostProcessSettings);
+
+	/** Overrides the post-process settings with the given ones. Only set properties are written. */
+	UFUNCTION(BlueprintCallable, Category=Camera)
+	static void OverrideChangedPostProcessSettings(const FBlueprintCameraEvaluationDataRef& CameraData, const FPostProcessSettings& PostProcessSettings);
+
+	/** Interpolate the post-process settings towards the given ones. Only set properties are interpolated. */
+	UFUNCTION(BlueprintCallable, Category=Camera)
+	static void LerpPostProcessSettings(const FBlueprintCameraEvaluationDataRef& CameraData, const FPostProcessSettings& PostProcessSettings, float BlendFactor);
+};
+
 #undef UE_API
diff --git a/Engine/Plugins/Cameras/GameplayCameras/Source/GameplayCameras/Public/Nodes/Common/PostProcessCameraNode.h b/Engine/Plugins/Cameras/GameplayCameras/Source/GameplayCameras/Public/Nodes/Common/PostProcessCameraNode.h
--- a/Engine/Plugins/Cameras/GameplayCameras/Source/GameplayCameras/Public/Nodes/Common/PostProcessCameraNode.h
+++ b/Engine/Plugins/Cameras/GameplayCameras/Source/GameplayCameras/Public/Nodes/Common/PostProcessCameraNode.h
@@ -23,7 +23,11 @@ protected:
 
 public:
 
-	UPROPERTY(EditAnywhere, Category="Rendering")
+	/** The post-process settings to set on the output camera pose. */
+	UPROPERTY(EditAnywhere, Category="Rendering", meta=(CameraContextData=true))
 	FPostProcessSettings PostProcessSettings;
+
+	UPROPERTY()
+	FCameraContextDataID PostProcessSettingsDataID;
 };
 
