Cameras: clean-up and improve handling of driven camera node parameters and multi-level prefabs

- Unify a bunch of code handling reflection-based node parameters vs custom node parameters (those obtained through the `ICustomCameraNodeParameterProvider` interface). Now there's the `FCameraNodeParameters` class that can build a list of available parameters on a camera node. This dramatically reduces the code in the Camera Rig builder class.
- Fix some issues with multi-level parameter overriding, such as when an exposed parameter is re-exposed a level higher. Also, automatically flag connected prefab parameters as "overriden" on the prefab camera node to make sure those overrides are taken into account.
- Rename some properties on custom parameter infos and interface parameter objects to be a bit clearer: "ParameterType" should be only for the parameter definitions, to specify if a parameter is "blendable" or "data". We should use "VariableType" for the type of variable to use in the variable table.
- Simplify some of the parameter override helper code and remove unnecessary helper methods.
- Add some unit tests for multi-level prefabs.

#Synced-CL 49781371

diff --git a/Engine/Plugins/Cameras/GameplayCameras/Config/DefaultGameplayCameras.ini b/Engine/Plugins/Cameras/GameplayCameras/Config/DefaultGameplayCameras.ini
--- a/Engine/Plugins/Cameras/GameplayCameras/Config/DefaultGameplayCameras.ini
+++ b/Engine/Plugins/Cameras/GameplayCameras/Config/DefaultGameplayCameras.ini
@@ -32,3 +32,10 @@
 ; UE 5.7 redirects.
 +PropertyRedirects=(OldName="/Script/GameplayCameras.CameraVariableAsset.bIsInput",NewName="CameraVariableAsset.bIsPreBlended")
 
+; UE 5.8 redirects.
++PropertyRedirects=(OldName="/Script/GameplayCameras.CustomCameraNodeBlendableParameter.ParameterType",NewName="VariableType")
++PropertyRedirects=(OldName="/Script/GameplayCameras.CustomCameraNodeDataParameter.ParameterType",NewName="DataType")
++PropertyRedirects=(OldName="/Script/GameplayCameras.CustomCameraNodeBlendableParameter.ParameterContainerType",NewName="DataContainerType")
++PropertyRedirects=(OldName="/Script/GameplayCameras.CustomCameraNodeBlendableParameter.ParameterTypeObject",NewName="DataTypeObject")
++PropertyRedirects=(OldName="/Script/GameplayCameras.CameraObjectInterfaceBlendableParameter.ParameterType",NewName="VariableType")
+
diff --git a/Engine/Plugins/Cameras/GameplayCameras/Source/GameplayCameras/Private/Build/CameraNodeHierarchyBuilder.cpp b/Engine/Plugins/Cameras/GameplayCameras/Source/GameplayCameras/Private/Build/CameraNodeHierarchyBuilder.cpp
--- a/Engine/Plugins/Cameras/GameplayCameras/Source/GameplayCameras/Private/Build/CameraNodeHierarchyBuilder.cpp
+++ b/Engine/Plugins/Cameras/GameplayCameras/Source/GameplayCameras/Private/Build/CameraNodeHierarchyBuilder.cpp
@@ -127,10 +127,10 @@ UE_CAMERA_VARIABLE_FOR_ALL_TYPES()
 	// Now do the same with custom parameters handled by the node itself.
 	if (ICustomCameraNodeParameterProvider* CustomParameterProvider = Cast<ICustomCameraNodeParameterProvider>(CameraNode))
 	{
-		FCustomCameraNodeParameterInfos CustomParameters;
-		CustomParameterProvider->GetCustomCameraNodeParameters(CustomParameters);
+		FCameraNodeParameterInfos CameraNodeParameters;
+		CustomParameterProvider->GetCustomCameraNodeParameters(CameraNodeParameters);
 
-		for (const FCustomCameraNodeParameterInfos::FBlendableParameterInfo& BlendableParameter : CustomParameters.BlendableParameters)
+		for (const FCameraNodeBlendableParameterInfo& BlendableParameter : CameraNodeParameters.BlendableParameters)
 		{
 			AddVariableToAllocationInfo(BlendableParameter.OverrideVariable, AllocationInfo.VariableTableInfo);
 		}
diff --git a/Engine/Plugins/Cameras/GameplayCameras/Source/GameplayCameras/Private/Build/CameraObjectInterfaceBuilder.cpp b/Engine/Plugins/Cameras/GameplayCameras/Source/GameplayCameras/Private/Build/CameraObjectInterfaceBuilder.cpp
--- a/Engine/Plugins/Cameras/GameplayCameras/Source/GameplayCameras/Private/Build/CameraObjectInterfaceBuilder.cpp
+++ b/Engine/Plugins/Cameras/GameplayCameras/Source/GameplayCameras/Private/Build/CameraObjectInterfaceBuilder.cpp
@@ -20,45 +20,10 @@ namespace UE::Cameras
 namespace Internal
 {
 
-template<typename CameraParameterOrVariableReference>
-bool MatchesVariableType(ECameraVariableType VariableType)
-{
-	return false;
-}
-
-template<typename CameraParameterOrVariableReference>
-FText GetVariableTypeAsText()
-{
-	return FText();
-}
-
-#define UE_CAMERA_VARIABLE_FOR_TYPE(ValueType, ValueName)\
-	template<>\
-	bool MatchesVariableType<F##ValueName##CameraVariableReference>(ECameraVariableType VariableType)\
-	{\
-		return VariableType == ECameraVariableType::ValueName;\
-	}\
-	template<>\
-	bool MatchesVariableType<F##ValueName##CameraParameter>(ECameraVariableType VariableType)\
-	{\
-		return VariableType == ECameraVariableType::ValueName;\
-	}\
-	template<>\
-	FText GetVariableTypeAsText<F##ValueName##CameraVariableReference>()\
-	{\
-		return UEnum::GetDisplayValueAsText(ECameraVariableType::ValueName);\
-	}\
-	template<>\
-	FText GetVariableTypeAsText<F##ValueName##CameraParameter>()\
-	{\
-		return UEnum::GetDisplayValueAsText(ECameraVariableType::ValueName);\
-	}
-UE_CAMERA_VARIABLE_FOR_ALL_TYPES()
-#undef UE_CAMERA_VARIABLE_FOR_TYPE
-
 struct FInterfaceParameterBindingBuilder
 {
 	UBaseCameraObject* CameraObject;
+	TMap<const UCameraNode*, FCameraNodeParameterInfos> BuiltCameraNodeParameters;
 
 	FInterfaceParameterBindingBuilder(FCameraObjectInterfaceBuilder& InOwner)
 		: Owner(InOwner)
@@ -76,439 +41,188 @@ struct FInterfaceParameterBindingBuilder
 		Owner.BuildLog.AddMessage(EMessageSeverity::Error, MoveTemp(ErrorMessage));
 	}
 
-	template<typename CameraParameterOrVariableReferenceType>
-	void SetCameraParameterOrVariableReferenceOverride(
-			const UCameraObjectInterfaceBlendableParameter* BlendableParameter,
+	void SetParameterOverride(
+			UCameraObjectInterfaceBlendableParameter* SourceParameter,
 			UCameraNode* TargetNode,
-			FStructProperty* TargetProperty,
-			CameraParameterOrVariableReferenceType* CameraParameterOrVariableReferencePtr)
+			FName TargetNodeParameterName)
 	{
-		using VariableAssetType = typename CameraParameterOrVariableReferenceType::VariableAssetType;
-		using ValueType = typename VariableAssetType::ValueType;
-
-		ensure(BlendableParameter->PrivateVariableID.IsValid());
-
-		const bool bIsValid = CheckIfParameterCanBeOverridden(BlendableParameter, TargetNode, TargetProperty->GetFName(), CameraParameterOrVariableReferencePtr);
-		if (!bIsValid)
+		const FCameraNodeBlendableParameterInfo* TargetParameterInfo = FindBlendableParameter(TargetNode, TargetNodeParameterName);
+		if (!TargetParameterInfo)
 		{
+			ReportError(
+					TargetNode,
+					FText::Format(LOCTEXT(
+							"InvalidConnectionTargetProperty",
+							"Invalid connection to property '{0}' on '{1}', no such property found."),
+						FText::FromName(TargetNodeParameterName),
+						FText::FromName(TargetNode->GetFName())));
 			return;
 		}
 
-		FCameraVariableID PreviousVariableID = FindOldDrivingVariableID(TargetProperty->GetFName(), TargetNode);
-		if (PreviousVariableID != BlendableParameter->PrivateVariableID)
+		if (TargetParameterInfo->VariableType != SourceParameter->VariableType ||
+				TargetParameterInfo->BlendableStructType != SourceParameter->BlendableStructType)
 		{
-			TargetNode->Modify();
+			ReportError(
+					TargetNode,
+					FText::Format(LOCTEXT(
+							"ConnectionVariableTypeMismatch",
+							"Invalid connection to property '{0}' on '{1}', expected type '{2}' but is '{3}'."),
+						FText::FromName(TargetNodeParameterName),
+						FText::FromName(TargetNode->GetFName()),
+						UEnum::GetDisplayValueAsText(SourceParameter->VariableType),
+						UEnum::GetDisplayValueAsText(TargetParameterInfo->VariableType)));
+			return;
 		}
 
-		CameraParameterOrVariableReferencePtr->VariableID = BlendableParameter->PrivateVariableID;
-	}
-
-	template<typename VariableAssetType>
-	void SetCustomBlendableParameterOverride(
-			const UCameraObjectInterfaceBlendableParameter* BlendableParameter,
-			UCameraNode* TargetNode,
-			FName TargetPropertyName,
-			const FCustomCameraNodeParameterInfos::FBlendableParameterInfo& CustomParameter)
-	{
-		using ValueType = typename VariableAssetType::ValueType;
-
-		ensure(BlendableParameter->PrivateVariableID.IsValid());
-		ensure(TargetPropertyName == CustomParameter.ParameterName);
+		if (!TargetParameterInfo->OverrideVariableID)
+		{
+			ReportError(
+					TargetNode,
+					FText::Format(LOCTEXT(
+							"InvalidConnectionTargetVariableID",
+							"Can't connect to property '{0}' on '{1}', no camera variable ID storage available."),
+						FText::FromName(TargetNodeParameterName),
+						FText::FromName(TargetNode->GetFName())));
+			return;
+		}
 
-		const bool bIsValid = CheckIfParameterCanBeOverridden(BlendableParameter, TargetNode, TargetPropertyName, CustomParameter);
-		if (!bIsValid)
+		if (TargetParameterInfo->OverrideVariableID->IsValid())
 		{
+			ReportError(
+					TargetNode,
+					FText::Format(LOCTEXT(
+							"ConnectionTargetVariableIDAlreadyUsed",
+							"Can't connect to property '{0}' on '{1}', it is already connected with something else."),
+						FText::FromName(TargetNodeParameterName),
+						FText::FromName(TargetNode->GetFName())));
 			return;
 		}
 
-		FCameraVariableID PreviousVariableID = FindOldDrivingVariableID(CustomParameter.ParameterName, TargetNode);
-		if (PreviousVariableID != BlendableParameter->PrivateVariableID)
+		ensure(SourceParameter->PrivateVariableID.IsValid());
+		FCameraVariableID PreviousVariableID = FindOldDrivingVariableID(TargetNode, TargetNodeParameterName);
+		if (PreviousVariableID != SourceParameter->PrivateVariableID)
 		{
 			TargetNode->Modify();
 		}
-
-		(*CustomParameter.OverrideVariableID) = BlendableParameter->PrivateVariableID;
+		*TargetParameterInfo->OverrideVariableID = SourceParameter->PrivateVariableID;
 	}
 
-	void SetCustomBlendableStructParameterOverride(
-			const UCameraObjectInterfaceBlendableParameter* BlendableParameter,
+	void SetParameterOverride(
+			UCameraObjectInterfaceDataParameter* SourceParameter,
 			UCameraNode* TargetNode,
-			FName TargetPropertyName,
-			const FCustomCameraNodeParameterInfos::FBlendableParameterInfo& CustomParameter)
+			FName TargetNodeParameterName)
 	{
-		ensure(BlendableParameter->PrivateVariableID.IsValid());
-		ensure(TargetPropertyName == CustomParameter.ParameterName);
-
-		const bool bIsValid = CheckIfParameterCanBeOverridden(BlendableParameter, TargetNode, TargetPropertyName, CustomParameter);
-		if (!bIsValid)
+		const FCameraNodeDataParameterInfo* TargetParameterInfo = FindDataParameter(TargetNode, TargetNodeParameterName);
+		if (!TargetParameterInfo)
 		{
+			ReportError(
+					TargetNode,
+					FText::Format(LOCTEXT(
+							"InvalidConnectionTargetProperty",
+							"Invalid connection to property '{0}' on '{1}', no such property found."),
+						FText::FromName(TargetNodeParameterName),
+						FText::FromName(TargetNode->GetFName())));
 			return;
 		}
 
-		// Also ensure the struct type is compatible.
-		if (CustomParameter.BlendableStructType != BlendableParameter->BlendableStructType)
+		if (TargetParameterInfo->DataType != SourceParameter->DataType ||
+				TargetParameterInfo->DataContainerType != SourceParameter->DataContainerType ||
+				TargetParameterInfo->DataTypeObject != SourceParameter->DataTypeObject)
 		{
-			ReportError(TargetNode,
-					FText::Format(
-						LOCTEXT(
-							"IncompatibleBlendableStructType", 
-							"Invalid interface parameter '{0}', driving property '{1}' on '{2}': expected type {3} but was {4}"),
-						FText::FromString(BlendableParameter->InterfaceParameterName),
-						FText::FromName(TargetPropertyName),
+			ReportError(
+					TargetNode,
+					FText::Format(LOCTEXT(
+							"ConnectionDataTypeMismatch",
+							"Invalid connection to property '{0}' on '{1}', expected type '{2}' but is '{3}'."),
+						FText::FromName(TargetNodeParameterName),
 						FText::FromName(TargetNode->GetFName()),
-#if WITH_EDITORONLY_DATA
-						CustomParameter.BlendableStructType->GetDisplayNameText(),
-						BlendableParameter->BlendableStructType->GetDisplayNameText()
-#else
-						FText::FromName(CustomParameter.BlendableStructType->GetFName()),
-						FText::FromName(BlendableParameter->BlendableStructType->GetFName())
-#endif
-						));
+						UEnum::GetDisplayValueAsText(SourceParameter->DataType),
+						UEnum::GetDisplayValueAsText(TargetParameterInfo->DataType)));
 			return;
 		}
 
-		FCameraVariableID PreviousVariableID = FindOldDrivingVariableID(CustomParameter.ParameterName, TargetNode);
-		if (PreviousVariableID != BlendableParameter->PrivateVariableID)
+		if (!TargetParameterInfo->OverrideDataID)
 		{
-			TargetNode->Modify();
+			ReportError(
+					TargetNode,
+					FText::Format(LOCTEXT(
+							"InvalidConnectionTargetDataID",
+							"Can't connect to property '{0}' on '{1}', no camera data ID storage available."),
+						FText::FromName(TargetNodeParameterName),
+						FText::FromName(TargetNode->GetFName())));
+			return;
 		}
 
-		(*CustomParameter.OverrideVariableID) = BlendableParameter->PrivateVariableID;
-	}
-
-	void SetDataContextPropertyOverride(
-			const UCameraObjectInterfaceDataParameter* DataParameter,
-			UCameraNode* TargetNode,
-			FProperty* TargetProperty,
-			FCameraContextDataID* OverrideDataID)
-	{
-		ensure(DataParameter->PrivateDataID);
-
-		const bool bIsValid = CheckIfParameterCanBeOverridden(DataParameter, TargetNode, TargetProperty, OverrideDataID);
-		if (!bIsValid)
+		if (TargetParameterInfo->OverrideDataID->IsValid())
 		{
+			ReportError(
+					TargetNode,
+					FText::Format(LOCTEXT(
+							"ConnectionTargetDataIDAlreadyUsed",
+							"Can't connect to property '{0}' on '{1}', it is already connected with something else."),
+						FText::FromName(TargetNodeParameterName),
+						FText::FromName(TargetNode->GetFName())));
 			return;
 		}
 
-		FCameraContextDataID PreviousDataID = FindOldDrivingDataID(TargetProperty->GetFName(), TargetNode);
-		if (PreviousDataID != DataParameter->PrivateDataID)
+		ensure(SourceParameter->PrivateDataID.IsValid());
+		FCameraContextDataID PreviousDataID = FindOldDrivingDataID(TargetNode, TargetNodeParameterName);
+		if (PreviousDataID != SourceParameter->PrivateDataID)
 		{
 			TargetNode->Modify();
 		}
-
-		(*OverrideDataID) = DataParameter->PrivateDataID;
+		*TargetParameterInfo->OverrideDataID = SourceParameter->PrivateDataID;
 	}
 
-	void SetCustomDataParameterOverride(
-			const UCameraObjectInterfaceDataParameter* DataParameter,
-			UCameraNode* TargetNode,
-			FName TargetPropertyName,
-			const FCustomCameraNodeParameterInfos::FDataParameterInfo& CustomParameter)
-	{
-		ensure(DataParameter->PrivateDataID);
-		ensure(TargetPropertyName == CustomParameter.ParameterName);
+private:
 
-		const bool bIsValid = CheckIfParameterCanBeOverridden(DataParameter, TargetNode, TargetPropertyName, CustomParameter);
-		if (!bIsValid)
+	const FCameraNodeParameterInfos& GetCameraNodeParameterInfos(UCameraNode* CameraNode)
+	{
+		const FCameraNodeParameterInfos* CameraNodeParameters = BuiltCameraNodeParameters.Find(CameraNode);
+		if (CameraNodeParameters)
 		{
-			return;
+			return *CameraNodeParameters;
 		}
+		
+		FCameraNodeParameterInfos& NewCameraNodeParameters = BuiltCameraNodeParameters.Add(CameraNode);
+		NewCameraNodeParameters.BuildFrom(CameraNode);
+		return NewCameraNodeParameters;
+	}
 
-		FCameraContextDataID PreviousDataID = FindOldDrivingDataID(CustomParameter.ParameterName, TargetNode);
-		if (PreviousDataID != DataParameter->PrivateDataID)
-		{
-			TargetNode->Modify();
-		}
+	const FCameraNodeBlendableParameterInfo* FindBlendableParameter(UCameraNode* CameraNode, FName ParameterName)
+	{
+		const FCameraNodeParameterInfos& CameraNodeParameters = GetCameraNodeParameterInfos(CameraNode);
+		return CameraNodeParameters.FindBlendableParameter(ParameterName);
+	}
 
-		(*CustomParameter.OverrideDataID) = DataParameter->PrivateDataID;
+	const FCameraNodeDataParameterInfo* FindDataParameter(UCameraNode* CameraNode, FName ParameterName)
+	{
+		const FCameraNodeParameterInfos& CameraNodeParameters = GetCameraNodeParameterInfos(CameraNode);
+		return CameraNodeParameters.FindDataParameter(ParameterName);
 	}
 
 private:
 
-	FCameraVariableID FindOldDrivingVariableID(FName ForParameterName, UObject* ForObject)
+	FCameraVariableID FindOldDrivingVariableID(UObject* ForObject, FName ForParameterName)
 	{
 		using FDrivenParameterKey = FCameraObjectInterfaceBuilder::FDrivenParameterKey;
 
-		FDrivenParameterKey ParameterKey{ ForParameterName, ForObject };
+		FDrivenParameterKey ParameterKey{ ForObject, ForParameterName };
 		FCameraVariableID ReusedVariableID;
 		Owner.OldDrivenBlendableParameters.RemoveAndCopyValue(ParameterKey, ReusedVariableID);
 		return ReusedVariableID;
 	}
 
-	FCameraContextDataID FindOldDrivingDataID(FName ForParameterName, UObject* ForObject)
+	FCameraContextDataID FindOldDrivingDataID(UObject* ForObject, FName ForParameterName)
 	{
 		using FDrivenParameterKey = FCameraObjectInterfaceBuilder::FDrivenParameterKey;
 
-		FDrivenParameterKey ParameterKey{ ForParameterName, ForObject };
+		FDrivenParameterKey ParameterKey{ ForObject, ForParameterName };
 		FCameraContextDataID ReusedDataID;
 		Owner.OldDrivenDataParameters.RemoveAndCopyValue(ParameterKey, ReusedDataID);
 		return ReusedDataID;
 	}
 
-	template<typename CameraParameterOrVariableReferenceType>
-	bool CheckIfParameterCanBeOverridden(
-			const UCameraObjectInterfaceBlendableParameter* BlendableParameter, 
-			UCameraNode* TargetNode,
-			const FName TargetPropertyName,
-			CameraParameterOrVariableReferenceType* CameraParameterOrVariableReference)
-	{
-		if (!MatchesVariableType<CameraParameterOrVariableReferenceType>(BlendableParameter->ParameterType))
-		{
-			ReportError(TargetNode,
-					FText::Format(
-						LOCTEXT(
-							"BlendableParameterTypeMismatch",
-							"Camera node parameter '{0}.{1}' of type {2} is not compatible with camera rig parameter '{3}' of type {4}"),
-						FText::FromName(TargetNode->GetFName()), 
-						FText::FromName(TargetPropertyName),
-						GetVariableTypeAsText<CameraParameterOrVariableReferenceType>(),
-						FText::FromString(BlendableParameter->InterfaceParameterName),
-						UEnum::GetDisplayValueAsText(BlendableParameter->ParameterType)));
-			return false;
-		}
-		if (CameraParameterOrVariableReference->Variable != nullptr)
-		{
-			ReportError(TargetNode,
-					FText::Format(
-						LOCTEXT(
-							"BlendableParameterDrivenTwice", 
-							"Camera node parameter '{0}.{1}' is both exposed and driven by a variable!"),
-						FText::FromName(TargetNode->GetFName()), 
-						FText::FromName(TargetPropertyName)));
-			return false;
-		}
-
-		return CheckIfParameterCanBeOverridden(BlendableParameter, TargetNode, TargetPropertyName, &CameraParameterOrVariableReference->VariableID);
-	}
-
-	bool CheckIfParameterCanBeOverridden(
-			const UCameraObjectInterfaceBlendableParameter* BlendableParameter, 
-			UCameraNode* TargetNode,
-			const FName TargetPropertyName,
-			const FCustomCameraNodeParameterInfos::FBlendableParameterInfo& CustomParameter)
-	{
-		if (BlendableParameter->ParameterType != CustomParameter.ParameterType)
-		{
-			ReportError(TargetNode,
-					FText::Format(
-						LOCTEXT(
-							"BlendableParameterTypeMismatch",
-							"Camera node parameter '{0}.{1}' of type {2} is not compatible with camera rig parameter '{3}' of type {4}"),
-						FText::FromName(TargetNode->GetFName()), 
-						FText::FromName(CustomParameter.ParameterName),
-						UEnum::GetDisplayValueAsText(CustomParameter.ParameterType),
-						FText::FromString(BlendableParameter->InterfaceParameterName),
-						UEnum::GetDisplayValueAsText(BlendableParameter->ParameterType)));
-			return false;
-		}
-		if (CustomParameter.OverrideVariable != nullptr)
-		{
-			ReportError(TargetNode,
-					FText::Format(
-						LOCTEXT(
-							"BlendableParameterDrivenTwice", 
-							"Camera node parameter '{0}.{1}' is both exposed and driven by a variable!"),
-						FText::FromName(TargetNode->GetFName()), 
-						FText::FromName(TargetPropertyName)));
-			return false;
-		}
-
-		return CheckIfParameterCanBeOverridden(BlendableParameter, TargetNode, TargetPropertyName, CustomParameter.OverrideVariableID);
-	}
-
-	bool CheckIfParameterCanBeOverridden(
-			const UCameraObjectInterfaceBlendableParameter* BlendableParameter, 
-			UCameraNode* TargetNode,
-			const FName TargetPropertyName,
-			FCameraVariableID* VariableID)
-	{
-		if (!VariableID)
-		{
-			ReportError(TargetNode,
-					FText::Format(
-						LOCTEXT(
-							"BlendableParameterMissingOverrideID", 
-							"Camera node parameter '{0}.{1}' cannot be overriden by a parameter"),
-						FText::FromName(TargetNode->GetFName()), 
-						FText::FromName(TargetPropertyName)));
-			return false;
-		}
-		if (VariableID->IsValid())
-		{
-			ReportError(TargetNode,
-					FText::Format(
-						LOCTEXT(
-							"BlendableParameterDrivenTwice", 
-							"Camera node parameter '{0}.{1}' is both exposed and driven by a variable!"),
-						FText::FromName(TargetNode->GetFName()), 
-						FText::FromName(TargetPropertyName)));
-			return false;
-		}
-		return true;
-	}
-
-	bool CheckIfParameterCanBeOverridden(
-			const UCameraObjectInterfaceDataParameter* DataParameter,
-			UCameraNode* TargetNode,
-			const FName TargetPropertyName,
-			const FCustomCameraNodeParameterInfos::FDataParameterInfo& CustomParameter)
-	{
-		if (DataParameter->DataType != CustomParameter.ParameterType)
-		{
-			ReportError(TargetNode,
-					FText::Format(
-						LOCTEXT(
-							"DataParameterTypeMismatch",
-							"Camera node parameter '{0}.{1}' of type {2} is not compatible with camera rig parameter '{3}' of type {4}"),
-						FText::FromName(TargetNode->GetFName()), 
-						FText::FromName(CustomParameter.ParameterName),
-						UEnum::GetDisplayValueAsText(CustomParameter.ParameterType),
-						FText::FromString(DataParameter->InterfaceParameterName),
-						UEnum::GetDisplayValueAsText(DataParameter->DataType)));
-			return false;
-		}
-		if (DataParameter->DataContainerType != CustomParameter.ParameterContainerType)
-		{
-			ReportError(TargetNode,
-					FText::Format(
-						LOCTEXT(
-							"DataParameterContainerTypeMismatch",
-							"Camera node parameter '{0}.{1}' has a different container type than camera rig parameter '{2}' ({3} vs {4})"),
-						FText::FromName(TargetNode->GetFName()), 
-						FText::FromName(CustomParameter.ParameterName),
-						FText::FromString(DataParameter->InterfaceParameterName),
-						UEnum::GetDisplayValueAsText(CustomParameter.ParameterContainerType),
-						UEnum::GetDisplayValueAsText(DataParameter->DataContainerType)));
-			return false;
-		}
-
-		return CheckIfParameterCanBeOverridden(DataParameter, TargetNode, TargetPropertyName, CustomParameter.OverrideDataID);
-	}
-
-	bool CheckIfParameterCanBeOverridden(
-			const UCameraObjectInterfaceDataParameter* DataParameter,
-			UCameraNode* TargetNode,
-			const FProperty* TargetProperty,
-			FCameraContextDataID* DataID)
-	{
-		bool bDataContainerTypeMatches = false;
-		switch (DataParameter->DataContainerType)
-		{
-			case ECameraContextDataContainerType::Array:
-				if (const FArrayProperty* ArrayProperty = CastField<FArrayProperty>(TargetProperty))
-				{
-					TargetProperty = ArrayProperty->Inner;
-					bDataContainerTypeMatches = true;
-				}
-				break;
-			default:
-				bDataContainerTypeMatches = 
-					!TargetProperty->IsA<FArrayProperty>() && 
-					!TargetProperty->IsA<FMapProperty>() && 
-					!TargetProperty->IsA<FSetProperty>();
-				break;
-		}
-		if (!bDataContainerTypeMatches)
-		{
-			ReportError(TargetNode,
-					FText::Format(
-						LOCTEXT(
-							"DataParameterContainerTypeMismatch",
-							"Camera node parameter '{0}.{1}' has a different container type than camera rig parameter '{2}' ({3} vs {4})"),
-						FText::FromName(TargetNode->GetFName()), 
-						FText::FromName(TargetProperty->GetFName()),
-						FText::FromString(DataParameter->InterfaceParameterName),
-						TargetProperty->GetClass()->GetDisplayNameText(),
-						UEnum::GetDisplayValueAsText(DataParameter->DataContainerType)));
-			return false;
-		}
-
-		bool bDataTypeMatches = false;
-		switch (DataParameter->DataType)
-		{
-			case ECameraContextDataType::Name:
-				bDataTypeMatches = (TargetProperty->IsA<FNameProperty>());
-				break;
-			case ECameraContextDataType::String:
-				bDataTypeMatches = (TargetProperty->IsA<FStrProperty>());
-				break;
-			case ECameraContextDataType::Enum:
-				if (const FEnumProperty* EnumProperty = CastField<FEnumProperty>(TargetProperty))
-				{
-					bDataTypeMatches = (EnumProperty->GetEnum() == DataParameter->DataTypeObject);
-				}
-				break;
-			case ECameraContextDataType::Struct:
-				if (const FStructProperty* StructProperty = CastField<FStructProperty>(TargetProperty))
-				{
-					bDataTypeMatches = (StructProperty->Struct == DataParameter->DataTypeObject);
-				}
-				break;
-			case ECameraContextDataType::Class:
-				if (const FClassProperty* ClassProperty = CastField<FClassProperty>(TargetProperty))
-				{
-					bDataTypeMatches = (ClassProperty->PropertyClass == DataParameter->DataTypeObject);
-				}
-				break;
-			case ECameraContextDataType::Object:
-				if (const FObjectProperty* ObjectProperty = CastField<FObjectProperty>(TargetProperty))
-				{
-					bDataTypeMatches = (ObjectProperty->PropertyClass == DataParameter->DataTypeObject);
-				}
-				break;
-			default:
-				ensure(false);
-				return false;
-		}
-		if (!bDataTypeMatches)
-		{
-			ReportError(TargetNode,
-					FText::Format(
-						LOCTEXT(
-							"DataParameterTypeMismatch",
-							"Camera node parameter '{0}.{1}' of type {2} is not compatible with camera rig parameter '{3}' of type {4}"),
-						FText::FromName(TargetNode->GetFName()), 
-						FText::FromName(TargetProperty->GetFName()),
-						TargetProperty->GetClass()->GetDisplayNameText(),
-						FText::FromString(DataParameter->InterfaceParameterName),
-						UEnum::GetDisplayValueAsText(DataParameter->DataType)));
-			return false;
-		}
-
-		return CheckIfParameterCanBeOverridden(DataParameter, TargetNode, TargetProperty->GetFName(), DataID);
-	}
-
-	bool CheckIfParameterCanBeOverridden(
-			const UCameraObjectInterfaceDataParameter* DataParameter,
-			UCameraNode* TargetNode,
-			const FName TargetPropertyName,
-			FCameraContextDataID* DataID)
-	{
-		if (!DataID)
-		{
-			ReportError(TargetNode,
-					FText::Format(
-						LOCTEXT(
-							"DataParameterMissingOverrideID", 
-							"Camera node parameter '{0}.{1}' cannot be overriden by a parameter"),
-						FText::FromName(TargetNode->GetFName()), 
-						FText::FromName(TargetPropertyName)));
-			return false;
-		}
-		if (DataID->IsValid())
-		{
-			ReportError(TargetNode,
-					FText::Format(
-						LOCTEXT(
-							"DataParameterDrivenTwice",
-							"Camera node parameter '{0}.{1}' is somehow overriden twice!"),
-						FText::FromName(TargetNode->GetFName()), 
-						FText::FromName(TargetPropertyName)));
-			return false;
-		}
-		return true;
-	}
-
 private:
 
 	FCameraObjectInterfaceBuilder& Owner;
@@ -587,104 +301,35 @@ void FCameraObjectInterfaceBuilder::GatherOldDrivenParameters()
 	OldDrivenBlendableParameters.Reset();
 	OldDrivenDataParameters.Reset();
 
+	FCameraNodeParameterInfos CameraNodeParameters;
+
 	for (UCameraNode* CameraNode : CameraObjectNodes)
 	{
-		UClass* CameraNodeClass = CameraNode->GetClass();
-		
-		for (TFieldIterator<FProperty> It(CameraNodeClass); It; ++It)
-		{
-			FProperty* Property(*It);
+		CameraNodeParameters.BuildFrom(CameraNode);
 
-			// First look for some blendable camera parameters.
-			if (FStructProperty* StructProperty = CastField<FStructProperty>(Property))
-			{
-				bool bFoundCameraParameter = true;
-#define UE_CAMERA_VARIABLE_FOR_TYPE(ValueType, ValueName)\
-				if (StructProperty->Struct == F##ValueName##CameraParameter::StaticStruct())\
-				{\
-					auto* CameraParameterPtr = StructProperty->ContainerPtrToValuePtr<F##ValueName##CameraParameter>(CameraNode);\
-					if (CameraParameterPtr->VariableID.IsValid() && !CameraParameterPtr->Variable)\
-					{\
-						OldDrivenBlendableParameters.Add(\
-								FDrivenParameterKey{ StructProperty->GetFName(), CameraNode },\
-								CameraParameterPtr->VariableID);\
-						CameraParameterPtr->VariableID = FCameraVariableID();\
-					}\
-				}\
-				else if (StructProperty->Struct == F##ValueName##CameraVariableReference::StaticStruct())\
-				{\
-					auto* VariableReferencePtr = StructProperty->ContainerPtrToValuePtr<F##ValueName##CameraVariableReference>(CameraNode);\
-					if (VariableReferencePtr->VariableID.IsValid() && !VariableReferencePtr->Variable)\
-					{\
-						OldDrivenBlendableParameters.Add(\
-								FDrivenParameterKey{ StructProperty->GetFName(), CameraNode },\
-								VariableReferencePtr->VariableID);\
-						VariableReferencePtr->VariableID = FCameraVariableID();\
-					}\
-				}\
-				else
-UE_CAMERA_VARIABLE_FOR_ALL_TYPES()
-#undef UE_CAMERA_VARIABLE_FOR_TYPE
-				{
-					// Other struct type...
-					bFoundCameraParameter = false;
-				}
-
-				if (bFoundCameraParameter)
-				{
-					continue;
-				}
-			}
-
-			// Then look for some data parameters.
-			const FName DataIDPropertyName = FName(It->GetName() + TEXT("DataID"));
-			FStructProperty* DataIDStructProperty = CastField<FStructProperty>(CameraNodeClass->FindPropertyByName(DataIDPropertyName));
-			if (DataIDStructProperty && DataIDStructProperty->Struct == FCameraContextDataID::StaticStruct())
+		for (const FCameraNodeBlendableParameterInfo& BlendableParameter : CameraNodeParameters.GetBlendableParameters())
+		{
+			if (BlendableParameter.OverrideVariableID &&
+					BlendableParameter.OverrideVariableID->IsValid() && 
+					// Don't touch user-defined variable overrides.
+					!BlendableParameter.OverrideVariable)
 			{
-				FCameraContextDataID* ExistingDataID = DataIDStructProperty->ContainerPtrToValuePtr<FCameraContextDataID>(CameraNode);
-				if (ExistingDataID->IsValid())
-				{
-					OldDrivenDataParameters.Add(FDrivenParameterKey{ Property->GetFName(), CameraNode }, *ExistingDataID);
-					*ExistingDataID = FCameraContextDataID();
-				}
+				OldDrivenBlendableParameters.Add(
+						FDrivenParameterKey{ CameraNode, BlendableParameter.ParameterName },
+						*BlendableParameter.OverrideVariableID);
+				*BlendableParameter.OverrideVariableID = FCameraVariableID();
 			}
 		}
 
-		if (ICustomCameraNodeParameterProvider* CustomParameterProvider = Cast<ICustomCameraNodeParameterProvider>(CameraNode))
+		for (const FCameraNodeDataParameterInfo& DataParameter : CameraNodeParameters.GetDataParameters())
 		{
-			FCustomCameraNodeParameterInfos CustomParameters;
-			CustomParameterProvider->GetCustomCameraNodeParameters(CustomParameters);
-
-			for (const FCustomCameraNodeParameterInfos::FBlendableParameterInfo& BlendableParameter : CustomParameters.BlendableParameters)
+			if (DataParameter.OverrideDataID && 
+					DataParameter.OverrideDataID->IsValid())
 			{
-				if (!BlendableParameter.OverrideVariableID)
-				{
-					continue;
-				}
-
-				if (BlendableParameter.OverrideVariableID->IsValid())
-				{
-					OldDrivenBlendableParameters.Add(
-							FDrivenParameterKey{ BlendableParameter.ParameterName, CameraNode },
-							*BlendableParameter.OverrideVariableID);
-					(*BlendableParameter.OverrideVariableID) = FCameraVariableID();
-				}
-			}
-
-			for (const FCustomCameraNodeParameterInfos::FDataParameterInfo& DataParameter : CustomParameters.DataParameters)
-			{
-				if (!DataParameter.OverrideDataID)
-				{
-					continue;
-				}
-
-				if (DataParameter.OverrideDataID->IsValid())
-				{
-					OldDrivenDataParameters.Add(
-							FDrivenParameterKey{ DataParameter.ParameterName, CameraNode },
-							*DataParameter.OverrideDataID);
-					(*DataParameter.OverrideDataID) = FCameraContextDataID();
-				}
+				OldDrivenDataParameters.Add(
+						FDrivenParameterKey{ CameraNode, DataParameter.ParameterName },
+						*DataParameter.OverrideDataID);
+				*DataParameter.OverrideDataID = FCameraContextDataID();
 			}
 		}
 	}
@@ -781,309 +426,39 @@ void FCameraObjectInterfaceBuilder::BuildInterfaceParameterBindings()
 	const FString CameraObjectName = CameraObject->GetName();
 	const FString CameraObjectPathName = CameraObject->GetPathName();
 
-	// Gather up all connected parameters, and what they're connected to.
-	using FCameraNodeProperty = TTuple<UCameraNode*, FName>;
-	TSet<FCameraNodeProperty> ConnectedCameraNodeProperties;
-	TMap<UCameraObjectInterfaceBlendableParameter*, TArray<FCameraNodeProperty>> BlendableParameterConnections;
-	TMap<UCameraObjectInterfaceDataParameter*, TArray<FCameraNodeProperty>> DataParameterConnections;
+	Internal::FInterfaceParameterBindingBuilder Builder(*this);
+
 	for (const FCameraObjectConnection& Connection : CameraObject->Connections.Connections)
 	{
-		if (Connection.SourcePropertyName != NAME_None)
-		{
-			continue;
-		}
-		if (Connection.Target == nullptr || Connection.TargetPropertyName == NAME_None)
+		if (Connection.Source == nullptr || Connection.Target == nullptr)
 		{
 			continue;
 		}
 
+		// For now limit ourselves to connections from interface parameters to camera nodes.
 		const UCameraObjectInterfaceParameterGetter* ParameterGetter = Cast<UCameraObjectInterfaceParameterGetter>(Connection.Source);
 		if (!ParameterGetter || !ParameterGetter->ParameterGuid.IsValid())
 		{
 			continue;
 		}
 
-		if (UCameraObjectInterfaceBlendableParameter* BlendableParameter = CameraObject->Interface.FindBlendableParameterByGuid(ParameterGetter->ParameterGuid))
-		{
-			BlendableParameterConnections.FindOrAdd(BlendableParameter).Add(FCameraNodeProperty{ Connection.Target, Connection.TargetPropertyName });
-		}
-		else if (UCameraObjectInterfaceDataParameter* DataParameter = CameraObject->Interface.FindDataParameterByGuid(ParameterGetter->ParameterGuid))
-		{
-			DataParameterConnections.FindOrAdd(DataParameter).Add(FCameraNodeProperty{ Connection.Target, Connection.TargetPropertyName });
-		}
-	}
-
-	auto CheckConnectionCollision = [this, &ConnectedCameraNodeProperties](const FCameraNodeProperty& InConnectedProperty) -> bool
-	{
-		bool bAlreadyConnected = false;
-		ConnectedCameraNodeProperties.Add(InConnectedProperty, &bAlreadyConnected);
-		if (bAlreadyConnected)
+		UCameraNode* TargetNode = Cast<UCameraNode>(Connection.Target);
+		if (!TargetNode)
 		{
-			BuildLog.AddMessage(EMessageSeverity::Error,
-					FText::Format(LOCTEXT(
-							"ParameterConnectionCollision",
-							"Multiple interface parameters targeting property '{0}' on camera node '{1}'. Ignoring duplicates."),
-						FText::FromName(InConnectedProperty.Value),
-						FText::FromName(InConnectedProperty.Key->GetFName())));
+			continue;
 		}
-		return bAlreadyConnected;
-	};
 
-	for (auto& Pair : BlendableParameterConnections)
-	{
-		const UCameraObjectInterfaceBlendableParameter* BlendableParameter(Pair.Key);
-		const TArray<FCameraNodeProperty>& ConnectedProperties(Pair.Value);
-	
-		for (const FCameraNodeProperty& ConnectedProperty : ConnectedProperties)
+		if (UCameraObjectInterfaceBlendableParameter* BlendableParameter = CameraObject->Interface.FindBlendableParameterByGuid(ParameterGetter->ParameterGuid))
 		{
-			// Check if some other parameter was already connected to this node's property.
-			if (CheckConnectionCollision(ConnectedProperty))
-			{
-				continue;
-			}
-
-			// See if this interface parameter is overriding a camera node parameter.
-			// Otherwise, maybe it's targeting a camera rig node's override for an inner rig interface parameter.
-			if (SetupCustomBlendableParameterOverride(BlendableParameter, ConnectedProperty.Key, ConnectedProperty.Value))
-			{
-				// Implicit continue.
-			}
-			else if (SetupCameraParameterOrVariableReferenceOverride(BlendableParameter, ConnectedProperty.Key, ConnectedProperty.Value))
-			{
-				// Implicit continue.
-			}
-			else
-			{
-				BuildLog.AddMessage(EMessageSeverity::Error,
-						ConnectedProperty.Key,
-						FText::Format(LOCTEXT(
-								"InvalidBlendableParameterTargetProperty",
-								"Invalid interface parameter '{0}', driving property '{1}' on '{2}', but no such property found."),
-							FText::FromString(BlendableParameter->InterfaceParameterName), 
-							FText::FromName(ConnectedProperty.Value),
-							FText::FromName(ConnectedProperty.Key->GetFName())));
-			}
+			Builder.SetParameterOverride(BlendableParameter, TargetNode, Connection.TargetPropertyName);
 		}
-	}
-
-	for (auto& Pair : DataParameterConnections)
-	{
-		const UCameraObjectInterfaceDataParameter* DataParameter(Pair.Key);
-		const TArray<FCameraNodeProperty>& ConnectedProperties(Pair.Value);
-
-		for (const FCameraNodeProperty& ConnectedProperty : ConnectedProperties)
+		else if (UCameraObjectInterfaceDataParameter* DataParameter = CameraObject->Interface.FindDataParameterByGuid(ParameterGetter->ParameterGuid))
 		{
-			// Check if some other parameter was already connected to this node's property.
-			if (CheckConnectionCollision(ConnectedProperty))
-			{
-				continue;
-			}
-
-			if (SetupCustomDataParameterOverride(DataParameter, ConnectedProperty.Key, ConnectedProperty.Value))
-			{
-				// Implicit continue.
-			}
-			else if (SetupDataContextPropertyOverride(DataParameter, ConnectedProperty.Key, ConnectedProperty.Value))
-			{
-				// Implicit continue.
-			}
-			else
-			{
-				BuildLog.AddMessage(EMessageSeverity::Error,
-						ConnectedProperty.Key,
-						FText::Format(LOCTEXT(
-								"InvalidDataParameterTargetProperty",
-								"Invalid interface parameter '{0}', driving property '{1}' on '{2}', but no such property found."),
-							FText::FromString(DataParameter->InterfaceParameterName), 
-							FText::FromName(ConnectedProperty.Value),
-							FText::FromName(ConnectedProperty.Key->GetFName())));
-			}
+			Builder.SetParameterOverride(DataParameter, TargetNode, Connection.TargetPropertyName);
 		}
 	}
 }
 
-bool FCameraObjectInterfaceBuilder::SetupCameraParameterOrVariableReferenceOverride(const UCameraObjectInterfaceBlendableParameter* BlendableParameter, UCameraNode* Target, FName TargetPropertyName)
-{
-	using namespace Internal;
-
-	// Here we hook up interface parameters connected to a camera node property. This property is supposed
-	// to be of one of the camera parameter types (FBooleanCameraParameter, FInteger32CameraParameter, etc.)
-	// so they have both a fixed value (bool, int32, etc.) and a "private variable" which we will set to the 
-	// private variable of the given interface parameter, checking that the types match (UBooleanCameraVariable, 
-	// UInteger32CameraVariable, etc.)
-
-	UClass* TargetClass = Target->GetClass();
-	FProperty* TargetProperty = TargetClass->FindPropertyByName(TargetPropertyName);
-	if (!TargetProperty)
-	{
-		// No match, try something else.
-		return false;
-	}
-
-	FStructProperty* TargetStructProperty = CastField<FStructProperty>(TargetProperty);
-	if (!TargetStructProperty)
-	{
-		BuildLog.AddMessage(EMessageSeverity::Error,
-				Target,
-				FText::Format(LOCTEXT(
-						"InvalidCameraNodeParameter",
-						"Invalid interface parameter '{0}', driving property '{1}' on '{2}', but it's not a camera parameter."),
-					FText::FromString(BlendableParameter->InterfaceParameterName), 
-					FText::FromName(TargetPropertyName),
-					FText::FromName(Target->GetFName())));
-		return true;
-	}
-
-	// Get the type of the camera parameter by matching the struct against all the types we support,
-	// and create a private camera variable asset to drive its value.
-	FInterfaceParameterBindingBuilder Builder(*this);
-#define UE_CAMERA_VARIABLE_FOR_TYPE(ValueType, ValueName)\
-	if (TargetStructProperty->Struct == F##ValueName##CameraParameter::StaticStruct())\
-	{\
-		auto* CameraParameterPtr = TargetStructProperty->ContainerPtrToValuePtr<F##ValueName##CameraParameter>(Target);\
-		Builder.SetCameraParameterOrVariableReferenceOverride<F##ValueName##CameraParameter>(\
-				BlendableParameter, Target, TargetStructProperty, CameraParameterPtr\
-				);\
-	}\
-	else if (TargetStructProperty->Struct == F##ValueName##CameraVariableReference::StaticStruct())\
-	{\
-		auto* VariableReferencePtr = TargetStructProperty->ContainerPtrToValuePtr<F##ValueName##CameraVariableReference>(Target);\
-		Builder.SetCameraParameterOrVariableReferenceOverride<F##ValueName##CameraVariableReference>(\
-				BlendableParameter, Target, TargetStructProperty, VariableReferencePtr\
-				);\
-	}\
-	else
-	UE_CAMERA_VARIABLE_FOR_ALL_TYPES()
-#undef UE_CAMERA_VARIABLE_FOR_TYPE
-	{
-		BuildLog.AddMessage(EMessageSeverity::Error,
-				BlendableParameter,
-				FText::Format(LOCTEXT(
-						"InvalidCameraNodeParameter",
-						"Invalid interface parameter '{0}', driving property '{1}' on '{2}', but it's not a camera parameter."),
-					FText::FromString(BlendableParameter->InterfaceParameterName), 
-					FText::FromName(TargetPropertyName),
-					FText::FromName(Target->GetFName())));
-	}
-
-	return true;
-}
-
-bool FCameraObjectInterfaceBuilder::SetupCustomBlendableParameterOverride(const UCameraObjectInterfaceBlendableParameter* BlendableParameter, UCameraNode* Target, FName TargetPropertyName)
-{
-	using namespace Internal;
-
-	ICustomCameraNodeParameterProvider* TargetProvider = Cast<ICustomCameraNodeParameterProvider>(Target);
-	if (!TargetProvider)
-	{
-		// No match, try something else.
-		return false;
-	}
-
-	// Look for a parameter override matching the target name.
-	// TODO: we're querying the list of custom parameters every time, we may want to cache it for this phase.
-	FCustomCameraNodeParameterInfos CustomParameters;
-	TargetProvider->GetCustomCameraNodeParameters(CustomParameters);
-
-	FCustomCameraNodeParameterInfos::FBlendableParameterInfo* TargetCustomParameter = 
-		CustomParameters.BlendableParameters.FindByPredicate(
-			[TargetPropertyName](FCustomCameraNodeParameterInfos::FBlendableParameterInfo& CustomParameter)
-			{
-				return CustomParameter.ParameterName == TargetPropertyName;
-			});
-	if (!TargetCustomParameter)
-	{
-		// No match, try something else.
-		return false;
-	}
-
-	FInterfaceParameterBindingBuilder Builder(*this);
-	switch (TargetCustomParameter->ParameterType)
-	{
-#define UE_CAMERA_VARIABLE_FOR_TYPE(ValueType, ValueName)\
-		case ECameraVariableType::ValueName:\
-			Builder.SetCustomBlendableParameterOverride<U##ValueName##CameraVariable>(\
-					BlendableParameter, Target, TargetPropertyName, *TargetCustomParameter);\
-			break;
-	UE_CAMERA_VARIABLE_FOR_ALL_TYPES()
-#undef UE_CAMERA_VARIABLE_FOR_TYPE
-		case ECameraVariableType::BlendableStruct:
-			Builder.SetCustomBlendableStructParameterOverride(BlendableParameter, Target, TargetPropertyName, *TargetCustomParameter);
-			break;
-	}
-
-	return true;
-}
-
-bool FCameraObjectInterfaceBuilder::SetupDataContextPropertyOverride(const UCameraObjectInterfaceDataParameter* DataParameter, UCameraNode* Target, FName TargetPropertyName)
-{
-	using namespace Internal;
-
-	UClass* TargetClass = Target->GetClass();
-	FProperty* TargetProperty = TargetClass->FindPropertyByName(TargetPropertyName);
-	if (!TargetProperty)
-	{
-		// No match, try something else.
-		return false;
-	}
-
-	const FName TargetDataIDPropertyName = FName(TargetProperty->GetName() + TEXT("DataID"));
-	FStructProperty* TargetDataIDProperty = CastField<FStructProperty>(TargetClass->FindPropertyByName(TargetDataIDPropertyName));
-	if (!TargetDataIDProperty || TargetDataIDProperty->Struct != FCameraContextDataID::StaticStruct())
-	{
-		BuildLog.AddMessage(EMessageSeverity::Error,
-				DataParameter,
-				FText::Format(LOCTEXT(
-						"MissingDataID", 
-						"Interface parameter '{0}' is driving data context property '{1}' on '{2}' "
-						"but no FCameraContextDataID property '{3}' was found to store the override ID."),
-					FText::FromString(DataParameter->InterfaceParameterName),
-					FText::FromName(TargetPropertyName),
-					FText::FromName(Target->GetFName()),
-					FText::FromName(TargetDataIDPropertyName)));
-		return false;
-	}
-
-	FCameraContextDataID* OverrideDataID = TargetDataIDProperty->ContainerPtrToValuePtr<FCameraContextDataID>(Target);
-
-	FInterfaceParameterBindingBuilder Builder(*this);
-	Builder.SetDataContextPropertyOverride(DataParameter, Target, TargetProperty, OverrideDataID);
-
-	return true;
-}
-
-bool FCameraObjectInterfaceBuilder::SetupCustomDataParameterOverride(const UCameraObjectInterfaceDataParameter* DataParameter, UCameraNode* Target, FName TargetPropertyName)
-{
-	using namespace Internal;
-
-	ICustomCameraNodeParameterProvider* TargetProvider = Cast<ICustomCameraNodeParameterProvider>(Target);
-	if (!TargetProvider)
-	{
-		// No match, try something else.
-		return false;
-	}
-
-	FCustomCameraNodeParameterInfos CustomParameters;
-	TargetProvider->GetCustomCameraNodeParameters(CustomParameters);
-
-	FCustomCameraNodeParameterInfos::FDataParameterInfo* TargetCustomParameter =
-		CustomParameters.DataParameters.FindByPredicate(
-				[TargetPropertyName](FCustomCameraNodeParameterInfos::FDataParameterInfo& CustomParameter)
-				{
-					return CustomParameter.ParameterName == TargetPropertyName;
-				});
-	if (!TargetCustomParameter)
-	{
-		// No match, try something else.
-		return false;
-	}
-
-	FInterfaceParameterBindingBuilder Builder(*this);
-	Builder.SetCustomDataParameterOverride(DataParameter, Target, TargetPropertyName, *TargetCustomParameter);
-
-	return true;
-}
-
 void FCameraObjectInterfaceBuilder::DiscardUnusedParameters()
 {
 	// Now that we've rebuilt all exposed parameters, anything left from the old list 
@@ -1092,14 +467,14 @@ void FCameraObjectInterfaceBuilder::DiscardUnusedParameters()
 
 	for (TPair<FDrivenParameterKey, FCameraVariableID> Pair : OldDrivenBlendableParameters)
 	{
-		UObject* Target = Pair.Key.Value;
+		UObject* Target = Pair.Key.Key;
 		Target->Modify();
 	}
 	OldDrivenBlendableParameters.Reset();
 
 	for (TPair<FDrivenParameterKey, FCameraContextDataID> Pair : OldDrivenDataParameters)
 	{
-		UObject* Target = Pair.Key.Value;
+		UObject* Target = Pair.Key.Key;
 		Target->Modify();
 	}
 	OldDrivenDataParameters.Reset();
diff --git a/Engine/Plugins/Cameras/GameplayCameras/Source/GameplayCameras/Private/Build/CameraObjectInterfaceParameterBuilder.cpp b/Engine/Plugins/Cameras/GameplayCameras/Source/GameplayCameras/Private/Build/CameraObjectInterfaceParameterBuilder.cpp
--- a/Engine/Plugins/Cameras/GameplayCameras/Source/GameplayCameras/Private/Build/CameraObjectInterfaceParameterBuilder.cpp
+++ b/Engine/Plugins/Cameras/GameplayCameras/Source/GameplayCameras/Private/Build/CameraObjectInterfaceParameterBuilder.cpp
@@ -72,8 +72,13 @@ void FCameraObjectInterfaceParameterBuilder::BuildDefaultParameters()
 	if (CameraObject->DefaultParameters.GetPropertyBagStruct() != DefaultParametersStruct)
 	{
 		CameraObject->Modify();
+		// In theory, the default values were set when the editor called SetDefaultParameterValue, which happens when the user
+		// creates a new parameter or changes that parameter's default value in the Details View. So the DefaultParameters
+		// structure should have the correct values already and we only need to migrate them to the new struct if anything
+		// changed.
 		CameraObject->DefaultParameters.MigrateToNewBagStruct(DefaultParametersStruct);
 	}
+
 }
 
 void FCameraObjectInterfaceParameterBuilder::AppendDefaultParameterProperties(const UBaseCameraObject* CameraObject, TArray<FPropertyBagPropertyDesc>& OutProperties)
@@ -189,12 +194,12 @@ bool FCameraObjectInterfaceParameterBuilder::SetDefaultParameterValue(UBaseCamer
 	// First check if the value is found on a custom parameter.
 	if (ICustomCameraNodeParameterProvider* CustomParameterProvider = Cast<ICustomCameraNodeParameterProvider>(TargetNode))
 	{
-		FCustomCameraNodeParameterInfos CustomParameters;
+		FCameraNodeParameterInfos CustomParameters;
 		CustomParameterProvider->GetCustomCameraNodeParameters(CustomParameters);
 
-		FCustomCameraNodeParameterInfos::FBlendableParameterInfo* CustomBlendableParameter = 
+		FCameraNodeBlendableParameterInfo* CustomBlendableParameter = 
 			CustomParameters.BlendableParameters.FindByPredicate(
-					[TargetPropertyName](FCustomCameraNodeParameterInfos::FBlendableParameterInfo& CustomParameter)
+					[TargetPropertyName](FCameraNodeBlendableParameterInfo& CustomParameter)
 					{
 						return CustomParameter.ParameterName == TargetPropertyName;
 					});
@@ -204,9 +209,9 @@ bool FCameraObjectInterfaceParameterBuilder::SetDefaultParameterValue(UBaseCamer
 			goto DoneSearchingRawSourceValuePtr;
 		}
 
-		FCustomCameraNodeParameterInfos::FDataParameterInfo* CustomDataParameter =
+		FCameraNodeDataParameterInfo* CustomDataParameter =
 			CustomParameters.DataParameters.FindByPredicate(
-					[TargetPropertyName](FCustomCameraNodeParameterInfos::FDataParameterInfo& CustomParameter)
+					[TargetPropertyName](FCameraNodeDataParameterInfo& CustomParameter)
 					{
 						return CustomParameter.ParameterName == TargetPropertyName;
 					});
diff --git a/Engine/Plugins/Cameras/GameplayCameras/Source/GameplayCameras/Private/Core/BaseCameraObjectReference.cpp b/Engine/Plugins/Cameras/GameplayCameras/Source/GameplayCameras/Private/Core/BaseCameraObjectReference.cpp
--- a/Engine/Plugins/Cameras/GameplayCameras/Source/GameplayCameras/Private/Core/BaseCameraObjectReference.cpp
+++ b/Engine/Plugins/Cameras/GameplayCameras/Source/GameplayCameras/Private/Core/BaseCameraObjectReference.cpp
@@ -17,21 +17,15 @@ const FCameraObjectInterfaceParameterMetaData* FBaseCameraObjectReference::FindM
 			});
 }
 
-FCameraObjectInterfaceParameterMetaData& FBaseCameraObjectReference::FindOrAddMetaData(const FGuid& PropertyID)
+FCameraObjectInterfaceParameterMetaData& FBaseCameraObjectReference::GetMetaData(const FGuid& PropertyID)
 {
-	FCameraObjectInterfaceParameterMetaData* ExistingMetaData = ParameterMetaData.FindByPredicate(
-			[PropertyID](FCameraObjectInterfaceParameterMetaData& Item)
+	FCameraObjectInterfaceParameterMetaData* MetaData = ParameterMetaData.FindByPredicate(
+			[PropertyID](const FCameraObjectInterfaceParameterMetaData& Item)
 			{
 				return Item.ParameterGuid == PropertyID;
 			});
-	if (ExistingMetaData)
-	{
-		return *ExistingMetaData;
-	}
-
-	FCameraObjectInterfaceParameterMetaData& NewMetaData = ParameterMetaData.Emplace_GetRef();
-	NewMetaData.ParameterGuid = PropertyID;
-	return NewMetaData;
+	check(MetaData);
+	return *MetaData;
 }
 
 bool FBaseCameraObjectReference::IsParameterOverridden(const FGuid& PropertyID) const
@@ -131,7 +125,7 @@ void FBaseCameraObjectReference::RebuildParameters()
 	}
 }
 
-void FBaseCameraObjectReference::GetCustomCameraNodeParameters(FCustomCameraNodeParameterInfos& OutParameterInfos)
+void FBaseCameraObjectReference::GetCustomCameraNodeParameters(UE::Cameras::FCameraNodeParameterInfos& OutParameterInfos)
 {
 	RebuildParametersIfNeeded();
 
@@ -196,7 +190,7 @@ void FBaseCameraObjectReference::GetCustomCameraNodeParameters(FCustomCameraNode
 						void* PropertyValue = PropertyDesc->CachedProperty->ContainerPtrToValuePtr<void>(ParametersMemory);
 						if (ensure(PropertyValue))
 						{
-							FCameraObjectInterfaceParameterMetaData& MetaData = FindOrAddMetaData(Definition.ParameterGuid);
+							FCameraObjectInterfaceParameterMetaData& MetaData = GetMetaData(Definition.ParameterGuid);
 							OutParameterInfos.AddBlendableParameter(
 									Definition.ParameterName,
 									Definition.VariableType,
@@ -210,7 +204,7 @@ void FBaseCameraObjectReference::GetCustomCameraNodeParameters(FCustomCameraNode
 		}
 		else if (Definition.ParameterType == ECameraObjectInterfaceParameterType::Data)
 		{
-			FCameraObjectInterfaceParameterMetaData& MetaData = FindOrAddMetaData(Definition.ParameterGuid);
+			FCameraObjectInterfaceParameterMetaData& MetaData = GetMetaData(Definition.ParameterGuid);
 
 			ECameraContextDataContainerType ContainerType = ECameraContextDataContainerType::None;
 			if (PropertyDesc->ContainerTypes.GetFirstContainerType() == EPropertyBagContainerType::Array)
@@ -267,6 +261,63 @@ void FBaseCameraObjectReference::GetCustomCameraNodeParameters(FCustomCameraNode
 	}
 }
 
+bool FBaseCameraObjectReference::PostBuild()
+{
+	const UPropertyBag* ParametersStruct = Parameters.GetPropertyBagStruct();
+	const uint8* ParametersMemory = Parameters.GetValue().GetMemory();
+	if (!ParametersStruct || !ParametersMemory)
+	{
+		return false;
+	}
+
+	// Parameters that are driven by a connection (such as a camera rig parameter) should be set as overriden.
+	// For most blendable parameters, we can tell about the connection from the VariableID being set on them.
+	// For blendable structs and data parameters, the VariableID or DataID is in our metadata.
+	bool bDidModify = false;
+	for (const FCameraObjectInterfaceParameterMetaData& MetaData : ParameterMetaData)
+	{
+		const FPropertyBagPropertyDesc* PropertyDesc = ParametersStruct->FindPropertyDescByID(MetaData.ParameterGuid);
+		if (!ensure(PropertyDesc))
+		{
+			continue;
+		}
+
+		if (PropertyDesc->ValueType == EPropertyBagPropertyType::Struct)
+		{
+#define UE_CAMERA_VARIABLE_FOR_TYPE(ValueType, ValueName)\
+			if (PropertyDesc->ValueTypeObject == F##ValueName##CameraParameter::StaticStruct())\
+			{\
+				using CameraParameterType = F##ValueName##CameraParameter;\
+				const CameraParameterType* PropertyValue = PropertyDesc->CachedProperty->ContainerPtrToValuePtr<CameraParameterType>(ParametersMemory);\
+				if (ensure(PropertyValue) && PropertyValue->VariableID.IsValid())\
+				{\
+					if (!Parameters.IsPropertyOverriden(MetaData.ParameterGuid))\
+					{\
+						Parameters.SetPropertyOverriden(MetaData.ParameterGuid, true);\
+						bDidModify = true;\
+					}\
+				}\
+			}\
+			else
+			UE_CAMERA_VARIABLE_FOR_ALL_TYPES()
+#undef UE_CAMERA_VARIABLE_FOR_TYPE
+			{
+				// Other struct, probably a blendable struct or data parameter, handled below.
+			}
+		}
+
+		if (MetaData.OverrideVariableID.IsValid() || MetaData.OverrideDataID.IsValid())
+		{
+			if (!Parameters.IsPropertyOverriden(MetaData.ParameterGuid))
+			{
+				Parameters.SetPropertyOverriden(MetaData.ParameterGuid, true);
+				bDidModify = true;
+			}
+		}
+	}
+	return bDidModify;
+}
+
 void FBaseCameraObjectReference::PostSerialize(const FArchive& Ar)
 {
 	if (Ar.IsLoading())
diff --git a/Engine/Plugins/Cameras/GameplayCameras/Source/GameplayCameras/Private/Core/BlendStackCameraNode.cpp b/Engine/Plugins/Cameras/GameplayCameras/Source/GameplayCameras/Private/Core/BlendStackCameraNode.cpp
--- a/Engine/Plugins/Cameras/GameplayCameras/Source/GameplayCameras/Private/Core/BlendStackCameraNode.cpp
+++ b/Engine/Plugins/Cameras/GameplayCameras/Source/GameplayCameras/Private/Core/BlendStackCameraNode.cpp
@@ -126,6 +126,10 @@ void FBlendStackCameraNodeEvaluator::InitializeEntry(
 		}
 	}
 
+	// Set default values for unset entries in the variable and context data tables.
+	FCameraObjectInterfaceParameterOverrideHelper Helper(NewEntry.ContextResult);
+	Helper.ApplyParameterDefaults(CameraRig);
+
 	// Initialize the node evaluators.
 	if (RootEvaluator)
 	{
@@ -140,10 +144,7 @@ void FBlendStackCameraNodeEvaluator::InitializeEntry(
 		RootEvaluator->Initialize(InitParams, NewEntry.ContextResult);  // Initializing with the context result here.
 	}
 
-	// Set default values for unset entries in the variable table, so that pre-blending from default 
-	// values works.
-	FCameraObjectInterfaceParameterOverrideHelper::ApplyDefaultBlendableParameters(CameraRig, NewEntry.ContextResult.VariableTable);
-
+	// Apply the initialized result.
 	NewEntry.Result.OverrideAll(NewEntry.ContextResult, true);
 
 	// Wrap up!
diff --git a/Engine/Plugins/Cameras/GameplayCameras/Source/GameplayCameras/Private/Core/BlendStackRootCameraNode.cpp b/Engine/Plugins/Cameras/GameplayCameras/Source/GameplayCameras/Private/Core/BlendStackRootCameraNode.cpp
--- a/Engine/Plugins/Cameras/GameplayCameras/Source/GameplayCameras/Private/Core/BlendStackRootCameraNode.cpp
+++ b/Engine/Plugins/Cameras/GameplayCameras/Source/GameplayCameras/Private/Core/BlendStackRootCameraNode.cpp
@@ -290,7 +290,8 @@ void FBlendStackRootCameraNodeEvaluator::RunBlendedParameterOverridesStack(const
 
 		// Start by setting the default values of all parameters. If we don't do this, parameter overrides
 		// wouldn't have a base value to blend from.
-		FCameraObjectInterfaceParameterOverrideHelper::ApplyDefaultBlendableParameters(BlendablePrefabCameraRig, CurResult.VariableTable);
+		FCameraObjectInterfaceParameterOverrideHelper Helper(&CurResult.VariableTable, nullptr);
+		Helper.ApplyParameterDefaults(BlendablePrefabCameraRig);
 
 		// Next, override the defaults with the specific values of this entry, applied bottoms up.
 		for (int32 Index = BlendedParameterOverrides.PrefabTrail.Num() - 1; Index >= 0; --Index)
diff --git a/Engine/Plugins/Cameras/GameplayCameras/Source/GameplayCameras/Private/Core/CameraAssetReference.cpp b/Engine/Plugins/Cameras/GameplayCameras/Source/GameplayCameras/Private/Core/CameraAssetReference.cpp
--- a/Engine/Plugins/Cameras/GameplayCameras/Source/GameplayCameras/Private/Core/CameraAssetReference.cpp
+++ b/Engine/Plugins/Cameras/GameplayCameras/Source/GameplayCameras/Private/Core/CameraAssetReference.cpp
@@ -22,9 +22,16 @@ void FCameraAssetReference::ApplyParameterOverrides(UE::Cameras::FCameraNodeEval
 	using namespace UE::Cameras;
 	if (CameraAsset)
 	{
-		FCameraObjectInterfaceParameterOverrideHelper Helper(&OutResult.VariableTable, &OutResult.ContextDataTable);
+		FCameraObjectInterfaceParameterOverrideHelper Helper(OutResult);
 		Helper.bDrivenOnly = bDrivenOnly;
-		Helper.ApplyParameterOverrides(CameraAsset, CameraAsset->GetParameterDefinitions(), Parameters);
+		Helper.ApplyFilteredParameters(
+				CameraAsset, 
+				CameraAsset->GetParameterDefinitions(),
+				Parameters,
+				[this](const FCameraObjectInterfaceParameterDefinition& Definition) -> bool
+				{
+					return Parameters.IsPropertyOverriden(Definition.ParameterGuid);
+				});
 	}
 }
 
diff --git a/Engine/Plugins/Cameras/GameplayCameras/Source/GameplayCameras/Private/Core/CameraObjectInterface.cpp b/Engine/Plugins/Cameras/GameplayCameras/Source/GameplayCameras/Private/Core/CameraObjectInterface.cpp
--- a/Engine/Plugins/Cameras/GameplayCameras/Source/GameplayCameras/Private/Core/CameraObjectInterface.cpp
+++ b/Engine/Plugins/Cameras/GameplayCameras/Source/GameplayCameras/Private/Core/CameraObjectInterface.cpp
@@ -43,7 +43,7 @@ FCameraVariableDefinition UCameraObjectInterfaceBlendableParameter::GetVariableD
 {
 	FCameraVariableDefinition Definition;
 	Definition.VariableID = PrivateVariableID;
-	Definition.VariableType = ParameterType;
+	Definition.VariableType = VariableType;
 	Definition.BlendableStructType = BlendableStructType;
 	Definition.bIsPrivate = true;
 	Definition.bIsInput = bIsPreBlended;
@@ -60,7 +60,7 @@ void UCameraObjectInterfaceBlendableParameter::GetParameterDefinition(FCameraObj
 	OutParameterDefinition.ParameterType = ECameraObjectInterfaceParameterType::Blendable;
 	OutParameterDefinition.bIsVisible = bIsVisible;
 	OutParameterDefinition.VariableID = PrivateVariableID;
-	OutParameterDefinition.VariableType = ParameterType;
+	OutParameterDefinition.VariableType = VariableType;
 	OutParameterDefinition.BlendableStructType = BlendableStructType;
 }
 
diff --git a/Engine/Plugins/Cameras/GameplayCameras/Source/GameplayCameras/Private/Core/CameraRigAsset.cpp b/Engine/Plugins/Cameras/GameplayCameras/Source/GameplayCameras/Private/Core/CameraRigAsset.cpp
--- a/Engine/Plugins/Cameras/GameplayCameras/Source/GameplayCameras/Private/Core/CameraRigAsset.cpp
+++ b/Engine/Plugins/Cameras/GameplayCameras/Source/GameplayCameras/Private/Core/CameraRigAsset.cpp
@@ -54,15 +54,15 @@ void UCameraRigAsset::PostLoad()
 
 #endif
 
-	// Initialize the ParameterType of blendable parameters that pre-existed the new interface data.
+	// Initialize the VariableType of blendable parameters that pre-existed the new interface data.
 	// The type defaults to Boolean so check only those.
 	for (UCameraObjectInterfaceBlendableParameter* BlendableParameter : Interface.BlendableParameters)
 	{
-		if (BlendableParameter->ParameterType == ECameraVariableType::Boolean)
+		if (BlendableParameter->VariableType == ECameraVariableType::Boolean)
 		{
 			if (BlendableParameter->PrivateVariable_DEPRECATED)
 			{
-				BlendableParameter->ParameterType = BlendableParameter->PrivateVariable_DEPRECATED->GetVariableType();
+				BlendableParameter->VariableType = BlendableParameter->PrivateVariable_DEPRECATED->GetVariableType();
 				BlendableParameter->PrivateVariable_DEPRECATED = nullptr;
 			}
 		}
@@ -285,7 +285,7 @@ void UCameraRigAsset::GetAssetRegistryTags(FAssetRegistryTagsContext Context) co
 		{
 			FAssetRegistryTag BlendableParameterTag;
 			BlendableParameterTag.Name = FName(BlendableParameter->InterfaceParameterName);
-			BlendableParameterTag.Value = VariableTypeEnum->GetNameStringByValue((int64)BlendableParameter->ParameterType);
+			BlendableParameterTag.Value = VariableTypeEnum->GetNameStringByValue((int64)BlendableParameter->VariableType);
 			BlendableParameterTag.Type = FAssetRegistryTag::TT_Alphabetical;
 			Context.AddTag(BlendableParameterTag);
 		}
diff --git a/Engine/Plugins/Cameras/GameplayCameras/Source/GameplayCameras/Private/Core/CameraRigAssetReference.cpp b/Engine/Plugins/Cameras/GameplayCameras/Source/GameplayCameras/Private/Core/CameraRigAssetReference.cpp
--- a/Engine/Plugins/Cameras/GameplayCameras/Source/GameplayCameras/Private/Core/CameraRigAssetReference.cpp
+++ b/Engine/Plugins/Cameras/GameplayCameras/Source/GameplayCameras/Private/Core/CameraRigAssetReference.cpp
@@ -29,20 +29,25 @@ const UBaseCameraObject* FCameraRigAssetReference::GetCameraObject() const
 
 void FCameraRigAssetReference::ApplyParameterOverrides(UE::Cameras::FCameraVariableTable& OutVariableTable, bool bDrivenOnly) const
 {
-	ApplyParameterOverridesImpl(&OutVariableTable, nullptr, bDrivenOnly);
+	ApplyParameterOverridesImpl(&OutVariableTable, nullptr, bDrivenOnly, false);
 }
 
 void FCameraRigAssetReference::ApplyParameterOverrides(UE::Cameras::FCameraVariableTable& OutVariableTable, UE::Cameras::FCameraContextDataTable& OutContextDataTable, bool bDrivenOnly) const
 {
-	ApplyParameterOverridesImpl(&OutVariableTable, &OutContextDataTable, bDrivenOnly);
+	ApplyParameterOverridesImpl(&OutVariableTable, &OutContextDataTable, bDrivenOnly, false);
 }
 
 void FCameraRigAssetReference::ApplyParameterOverrides(UE::Cameras::FCameraNodeEvaluationResult& OutResult, bool bDrivenOnly) const
 {
-	ApplyParameterOverridesImpl(&OutResult.VariableTable, &OutResult.ContextDataTable, bDrivenOnly);
+	ApplyParameterOverridesImpl(&OutResult.VariableTable, &OutResult.ContextDataTable, bDrivenOnly, false);
 }
 
-void FCameraRigAssetReference::ApplyParameterOverridesImpl(UE::Cameras::FCameraVariableTable* OutVariableTable, UE::Cameras::FCameraContextDataTable* OutContextDataTable, bool bDrivenOnly) const
+void FCameraRigAssetReference::ApplyParameterOverridesAndDefaults(UE::Cameras::FCameraNodeEvaluationResult& OutResult) const
+{
+	ApplyParameterOverridesImpl(&OutResult.VariableTable, &OutResult.ContextDataTable, false, true);
+}
+
+void FCameraRigAssetReference::ApplyParameterOverridesImpl(UE::Cameras::FCameraVariableTable* OutVariableTable, UE::Cameras::FCameraContextDataTable* OutContextDataTable, bool bDrivenOnly, bool bApplyUnwrittenDefaults) const
 {
 	using namespace UE::Cameras;
 
@@ -50,7 +55,11 @@ void FCameraRigAssetReference::ApplyParameterOverridesImpl(UE::Cameras::FCameraV
 	{
 		FCameraObjectInterfaceParameterOverrideHelper Helper(OutVariableTable, OutContextDataTable);
 		Helper.bDrivenOnly = bDrivenOnly;
-		Helper.ApplyParameterOverrides(CameraRig, CameraRig->GetParameterDefinitions(), Parameters);
+		Helper.ApplyParameterOverrides(CameraRig, Parameters);
+		if (bApplyUnwrittenDefaults)
+		{
+			Helper.ApplyParameterDefaults(CameraRig, true);
+		}
 	}
 }
 
diff --git a/Engine/Plugins/Cameras/GameplayCameras/Source/GameplayCameras/Private/Core/CameraShakeAssetReference.cpp b/Engine/Plugins/Cameras/GameplayCameras/Source/GameplayCameras/Private/Core/CameraShakeAssetReference.cpp
--- a/Engine/Plugins/Cameras/GameplayCameras/Source/GameplayCameras/Private/Core/CameraShakeAssetReference.cpp
+++ b/Engine/Plugins/Cameras/GameplayCameras/Source/GameplayCameras/Private/Core/CameraShakeAssetReference.cpp
@@ -24,20 +24,25 @@ const UBaseCameraObject* FCameraShakeAssetReference::GetCameraObject() const
 
 void FCameraShakeAssetReference::ApplyParameterOverrides(UE::Cameras::FCameraVariableTable& OutVariableTable, bool bDrivenOnly) const
 {
-	ApplyParameterOverridesImpl(&OutVariableTable, nullptr, bDrivenOnly);
+	ApplyParameterOverridesImpl(&OutVariableTable, nullptr, bDrivenOnly, false);
 }
 
 void FCameraShakeAssetReference::ApplyParameterOverrides(UE::Cameras::FCameraVariableTable& OutVariableTable, UE::Cameras::FCameraContextDataTable& OutContextDataTable, bool bDrivenOnly) const
 {
-	ApplyParameterOverridesImpl(&OutVariableTable, &OutContextDataTable, bDrivenOnly);
+	ApplyParameterOverridesImpl(&OutVariableTable, &OutContextDataTable, bDrivenOnly, false);
 }
 
 void FCameraShakeAssetReference::ApplyParameterOverrides(UE::Cameras::FCameraNodeEvaluationResult& OutResult, bool bDrivenOnly) const
 {
-	ApplyParameterOverridesImpl(&OutResult.VariableTable, &OutResult.ContextDataTable, bDrivenOnly);
+	ApplyParameterOverridesImpl(&OutResult.VariableTable, &OutResult.ContextDataTable, bDrivenOnly, false);
 }
 
-void FCameraShakeAssetReference::ApplyParameterOverridesImpl(UE::Cameras::FCameraVariableTable* OutVariableTable, UE::Cameras::FCameraContextDataTable* OutContextDataTable, bool bDrivenOnly) const
+void FCameraShakeAssetReference::ApplyParameterOverridesAndDefaults(UE::Cameras::FCameraNodeEvaluationResult& OutResult) const
+{
+	ApplyParameterOverridesImpl(&OutResult.VariableTable, &OutResult.ContextDataTable, false, true);
+}
+
+void FCameraShakeAssetReference::ApplyParameterOverridesImpl(UE::Cameras::FCameraVariableTable* OutVariableTable, UE::Cameras::FCameraContextDataTable* OutContextDataTable, bool bDrivenOnly, bool bApplyUnwrittenDefaults) const
 {
 	using namespace UE::Cameras;
 	
@@ -45,7 +50,11 @@ void FCameraShakeAssetReference::ApplyParameterOverridesImpl(UE::Cameras::FCamer
 	{
 		FCameraObjectInterfaceParameterOverrideHelper Helper(OutVariableTable, OutContextDataTable);
 		Helper.bDrivenOnly = bDrivenOnly;
-		Helper.ApplyParameterOverrides(CameraShake, CameraShake->GetParameterDefinitions(), Parameters);
+		Helper.ApplyParameterOverrides(CameraShake, Parameters);
+		if (bApplyUnwrittenDefaults)
+		{
+			Helper.ApplyParameterDefaults(CameraShake, true);
+		}
 	}
 }
 
diff --git a/Engine/Plugins/Cameras/GameplayCameras/Source/GameplayCameras/Private/Core/ICustomCameraNodeParameterProvider.cpp b/Engine/Plugins/Cameras/GameplayCameras/Source/GameplayCameras/Private/Core/ICustomCameraNodeParameterProvider.cpp
--- a/Engine/Plugins/Cameras/GameplayCameras/Source/GameplayCameras/Private/Core/ICustomCameraNodeParameterProvider.cpp
+++ b/Engine/Plugins/Cameras/GameplayCameras/Source/GameplayCameras/Private/Core/ICustomCameraNodeParameterProvider.cpp
@@ -2,123 +2,307 @@
 
 #include "Core/ICustomCameraNodeParameterProvider.h"
 
+#include "Core/CameraNode.h"
+#include "Core/CameraParameters.h"
+#include "Core/CameraVariableReferences.h"
 #include "GameplayCamerasDelegates.h"
 
 #include UE_INLINE_GENERATED_CPP_BY_NAME(ICustomCameraNodeParameterProvider)
 
-void FCustomCameraNodeParameterInfos::AddBlendableParameter(
+namespace UE::Cameras
+{
+
+void FCameraNodeParameterInfos::AddBlendableParameter(
 		FName ParameterName, 
-		ECameraVariableType ParameterType, 
+		ECameraVariableType VariableType, 
 		const UScriptStruct* BlendableStructType,
 		const uint8* DefaultValue,
 		FCameraVariableID* OverrideVariableID)
 {
-	BlendableParameters.Add({ ParameterName, ParameterType, BlendableStructType, DefaultValue, OverrideVariableID });
+	BlendableParameters.Add({ ParameterName, VariableType, BlendableStructType, DefaultValue, OverrideVariableID });
 }
 
-void FCustomCameraNodeParameterInfos::AddBlendableParameter(FCustomCameraNodeBlendableParameter& Parameter, const uint8* DefaultValue)
+void FCameraNodeParameterInfos::AddBlendableParameter(FCustomCameraNodeBlendableParameter& Parameter, const uint8* DefaultValue)
 {
 	AddBlendableParameter(
 			Parameter.ParameterName,
-			Parameter.ParameterType,
+			Parameter.VariableType,
 			Parameter.BlendableStructType,
 			DefaultValue,
 			Parameter.OverrideVariable ? nullptr : &Parameter.OverrideVariableID);
 }
 
-void FCustomCameraNodeParameterInfos::AddDataParameter(
+void FCameraNodeParameterInfos::AddDataParameter(
 		FName ParameterName, 
-		ECameraContextDataType ParameterType,
-		ECameraContextDataContainerType ParameterContainerType,
-		const UObject* ParameterTypeObject,
+		ECameraContextDataType DataType,
+		ECameraContextDataContainerType DataContainerType,
+		const UObject* DataTypeObject,
 		const uint8* DefaultValue,
 		FCameraContextDataID* OverrideDataID)
 {
-	DataParameters.Add({ ParameterName, ParameterType, ParameterContainerType, ParameterTypeObject, DefaultValue, OverrideDataID });
+	DataParameters.Add({ ParameterName, DataType, DataContainerType, DataTypeObject, DefaultValue, OverrideDataID });
 }
 
-void FCustomCameraNodeParameterInfos::AddDataParameter(FCustomCameraNodeDataParameter& Parameter, const uint8* DefaultValue)
+void FCameraNodeParameterInfos::AddDataParameter(FCustomCameraNodeDataParameter& Parameter, const uint8* DefaultValue)
 {
 	AddDataParameter(
 			Parameter.ParameterName,
-			Parameter.ParameterType,
-			Parameter.ParameterContainerType,
-			Parameter.ParameterTypeObject,
+			Parameter.DataType,
+			Parameter.DataContainerType,
+			Parameter.DataTypeObject,
 			DefaultValue,
 			&Parameter.OverrideDataID);
 }
 
-void FCustomCameraNodeParameterInfos::GetBlendableParameters(TArray<FCustomCameraNodeBlendableParameter>& OutBlendableParameters) const
+void FCameraNodeParameterInfos::Reset()
+{
+	BlendableParameters.Reset();
+	DataParameters.Reset();
+}
+
+const FCameraNodeBlendableParameterInfo* FCameraNodeParameterInfos::FindBlendableParameter(FName ParameterName) const
 {
 	for (const FBlendableParameterInfo& BlendableParameter : BlendableParameters)
 	{
-		FCustomCameraNodeBlendableParameter& OutParameter = OutBlendableParameters.Emplace_GetRef();
-		OutParameter.ParameterName = BlendableParameter.ParameterName;
-		OutParameter.ParameterType = BlendableParameter.ParameterType;
-		OutParameter.BlendableStructType = BlendableParameter.BlendableStructType;
-		OutParameter.OverrideVariable = BlendableParameter.OverrideVariable;
-		if (BlendableParameter.OverrideVariableID)
+		if (BlendableParameter.ParameterName == ParameterName)
 		{
-			OutParameter.OverrideVariableID = *BlendableParameter.OverrideVariableID;
+			return &BlendableParameter;
 		}
 	}
+	return nullptr;
 }
 
-void FCustomCameraNodeParameterInfos::GetDataParameters(TArray<FCustomCameraNodeDataParameter>& OutDataParameters) const
+const FCameraNodeDataParameterInfo* FCameraNodeParameterInfos::FindDataParameter(FName ParameterName) const
 {
 	for (const FDataParameterInfo& DataParameter : DataParameters)
 	{
-		FCustomCameraNodeDataParameter& OutParameter = OutDataParameters.Emplace_GetRef();
-		OutParameter.ParameterName = DataParameter.ParameterName;
-		OutParameter.ParameterType = DataParameter.ParameterType;
-		OutParameter.ParameterContainerType = DataParameter.ParameterContainerType;
-		OutParameter.ParameterTypeObject = DataParameter.ParameterTypeObject;
-		if (DataParameter.OverrideDataID)
+		if (DataParameter.ParameterName == ParameterName)
 		{
-			OutParameter.OverrideDataID = *DataParameter.OverrideDataID;
+			return &DataParameter;
 		}
 	}
+	return nullptr;
 }
 
-bool FCustomCameraNodeParameterInfos::FindBlendableParameter(FName ParameterName, FCustomCameraNodeBlendableParameter& OutParameter) const
+void FCameraNodeParameterInfos::BuildFrom(UCameraNode* InCameraNode)
 {
-	for (const FBlendableParameterInfo& BlendableParameter : BlendableParameters)
+	Reset();
+
+	if (!ensure(InCameraNode))
 	{
-		if (BlendableParameter.ParameterName == ParameterName)
+		return;
+	}
+
+	UClass* CameraNodeClass = InCameraNode->GetClass();
+
+	for (TFieldIterator<FProperty> It(CameraNodeClass); It; ++It)
+	{
+		FProperty* Property(*It);
+
+		// First look for some built-in blendable parameters.
+		if (FStructProperty* StructProperty = CastField<FStructProperty>(Property))
 		{
-			OutParameter.ParameterName = BlendableParameter.ParameterName;
-			OutParameter.ParameterType = BlendableParameter.ParameterType;
-			OutParameter.BlendableStructType = BlendableParameter.BlendableStructType;
-			if (BlendableParameter.OverrideVariableID)
+			bool bIsCameraParameterProperty = true;
+#define UE_CAMERA_VARIABLE_FOR_TYPE(ValueType, ValueName)\
+			if (StructProperty->Struct == F##ValueName##CameraParameter::StaticStruct())\
+			{\
+				auto* CameraParameterPtr = StructProperty->ContainerPtrToValuePtr<F##ValueName##CameraParameter>(InCameraNode);\
+				AddBlendableParameter(\
+						StructProperty->GetFName(),\
+						ECameraVariableType::ValueName,\
+						nullptr,\
+						reinterpret_cast<uint8*>(&CameraParameterPtr->Value),\
+						&CameraParameterPtr->VariableID);\
+			}\
+			else if (StructProperty->Struct == F##ValueName##CameraVariableReference::StaticStruct())\
+			{\
+				auto* VariableReferencePtr = StructProperty->ContainerPtrToValuePtr<F##ValueName##CameraVariableReference>(InCameraNode);\
+				AddBlendableParameter(\
+						StructProperty->GetFName(),\
+						ECameraVariableType::ValueName,\
+						nullptr,\
+						nullptr,\
+						&VariableReferencePtr->VariableID);\
+			}\
+			else
+			UE_CAMERA_VARIABLE_FOR_ALL_TYPES()
+#undef UE_CAMERA_VARIABLE_FOR_TYPE
+			{
+				// Other struct type...
+				bIsCameraParameterProperty = false;
+			}
+
+			if (bIsCameraParameterProperty)
 			{
-				OutParameter.OverrideVariableID = *BlendableParameter.OverrideVariableID;
+				continue;
 			}
-			return true;
 		}
-	}
-	return false;
-}
 
-bool FCustomCameraNodeParameterInfos::FindDataParameter(FName ParameterName, FCustomCameraNodeDataParameter& OutParameter) const
-{
-	for (const FDataParameterInfo& DataParameter : DataParameters)
-	{
-		if (DataParameter.ParameterName == ParameterName)
+		// Look for a custom blendable parameter.
+		const FName BlendableIDPropertyName = FName(It->GetName() + TEXT("BlendableID"));
+		FStructProperty* BlendableIDStructProperty = CastField<FStructProperty>(CameraNodeClass->FindPropertyByName(BlendableIDPropertyName));
+		if (BlendableIDStructProperty && BlendableIDStructProperty->Struct == FCameraVariableID::StaticStruct())
 		{
-			OutParameter.ParameterName = DataParameter.ParameterName;
-			OutParameter.ParameterType = DataParameter.ParameterType;
-			OutParameter.ParameterContainerType = DataParameter.ParameterContainerType;
-			OutParameter.ParameterTypeObject = DataParameter.ParameterTypeObject;
-			if (DataParameter.OverrideDataID)
+			bool bIsCameraParameterProperty = true;
+			ECameraVariableType VariableType = ECameraVariableType::Boolean;
+			UScriptStruct* VariableTypeObject = nullptr;
+			if (Property->IsA<FBoolProperty>())
+			{
+				VariableType = ECameraVariableType::Boolean;
+			}
+			else if (Property->IsA<FIntProperty>())
+			{
+				VariableType = ECameraVariableType::Integer32;
+			}
+			else if (Property->IsA<FFloatProperty>())
+			{
+				VariableType = ECameraVariableType::Float;
+			}
+			else if (Property->IsA<FDoubleProperty>())
+			{
+				VariableType = ECameraVariableType::Double;
+			}
+			else if (FStructProperty* StructProperty = CastField<FStructProperty>(Property))
 			{
-				OutParameter.OverrideDataID = *DataParameter.OverrideDataID;
+				if (StructProperty->Struct == TVariantStructure<FVector2f>::Get())
+				{
+					VariableType = ECameraVariableType::Vector2f;
+				}
+				else if (StructProperty->Struct == TBaseStructure<FVector2D>::Get())
+				{
+					VariableType = ECameraVariableType::Vector2d;
+				}
+				else if (StructProperty->Struct == TVariantStructure<FVector3f>::Get())
+				{
+					VariableType = ECameraVariableType::Vector3f;
+				}
+				else if (StructProperty->Struct == TBaseStructure<FVector>::Get())
+				{
+					VariableType = ECameraVariableType::Vector3d;
+				}
+				else if (StructProperty->Struct == TVariantStructure<FVector4f>::Get())
+				{
+					VariableType = ECameraVariableType::Vector4f;
+				}
+				else if (StructProperty->Struct == TBaseStructure<FVector4>::Get())
+				{
+					VariableType = ECameraVariableType::Vector4d;
+				}
+				else if (StructProperty->Struct == TVariantStructure<FRotator3f>::Get())
+				{
+					VariableType = ECameraVariableType::Rotator3f;
+				}
+				else if (StructProperty->Struct == TBaseStructure<FRotator>::Get())
+				{
+					VariableType = ECameraVariableType::Rotator3d;
+				}
+				else if (StructProperty->Struct == TVariantStructure<FTransform3f>::Get())
+				{
+					VariableType = ECameraVariableType::Transform3f;
+				}
+				else if (StructProperty->Struct == TBaseStructure<FTransform>::Get())
+				{
+					VariableType = ECameraVariableType::Transform3d;
+				}
+				else
+				{
+					// TODO: check that the structure is registered as blendable
+					VariableType = ECameraVariableType::BlendableStruct;
+					VariableTypeObject = StructProperty->Struct;
+				}
+			}
+			else
+			{
+				bIsCameraParameterProperty = false;
+			}
+
+			if (bIsCameraParameterProperty)
+			{
+				void* DefaultValue = Property->ContainerPtrToValuePtr<void>(InCameraNode);
+				FCameraVariableID* BlendableID = BlendableIDStructProperty->ContainerPtrToValuePtr<FCameraVariableID>(InCameraNode);
+				AddBlendableParameter(
+						Property->GetFName(),
+						VariableType,
+						VariableTypeObject,
+						static_cast<uint8*>(DefaultValue),
+						BlendableID);
+				continue;
+			}
+		}
+
+		// Look for a data parameter.
+		const FName DataIDPropertyName = FName(It->GetName() + TEXT("DataID"));
+		FStructProperty* DataIDStructProperty = CastField<FStructProperty>(CameraNodeClass->FindPropertyByName(DataIDPropertyName));
+		if (DataIDStructProperty && DataIDStructProperty->Struct == FCameraContextDataID::StaticStruct())
+		{
+			bool bIsDataProperty = true;
+			ECameraContextDataType DataPropertyType = ECameraContextDataType::Name;
+			ECameraContextDataContainerType DataPropertyContainerType = ECameraContextDataContainerType::None;
+			UObject* DataPropertyTypeObject = nullptr;
+
+			bool bHasDataPropertyContainer = false;
+			FProperty* ActualProperty = Property;
+			if (FArrayProperty* ArrayProperty = CastField<FArrayProperty>(Property))
+			{
+				bHasDataPropertyContainer = true;
+				ActualProperty = ArrayProperty->Inner;
+				DataPropertyContainerType = ECameraContextDataContainerType::Array;
+			}
+
+			if (ActualProperty->IsA<FNameProperty>())
+			{
+				DataPropertyType = ECameraContextDataType::Name;
+			}
+			else if (ActualProperty->IsA<FStrProperty>())
+			{
+				DataPropertyType = ECameraContextDataType::String;
+			}
+			else if (FEnumProperty* EnumProperty = CastField<FEnumProperty>(ActualProperty))
+			{
+				DataPropertyType = ECameraContextDataType::Enum;
+				DataPropertyTypeObject = EnumProperty->GetEnum();
+			}
+			else if (FClassProperty* ClassProperty = CastField<FClassProperty>(ActualProperty))
+			{
+				DataPropertyType = ECameraContextDataType::Class;
+				DataPropertyTypeObject = ClassProperty->MetaClass;
+			}
+			else if (FObjectProperty* ObjectProperty = CastField<FObjectProperty>(ActualProperty))
+			{
+				DataPropertyType = ECameraContextDataType::Object;
+				DataPropertyTypeObject = ObjectProperty->PropertyClass;
+			}
+			else
+			{
+				bIsDataProperty = false;
+			}
+
+			if (bIsDataProperty)
+			{
+				void* DefaultValue = !bHasDataPropertyContainer ? 
+					Property->ContainerPtrToValuePtr<void>(InCameraNode) : nullptr;
+				FCameraContextDataID* DataID = DataIDStructProperty->ContainerPtrToValuePtr<FCameraContextDataID>(InCameraNode);
+				AddDataParameter(
+						Property->GetFName(),
+						DataPropertyType,
+						DataPropertyContainerType,
+						DataPropertyTypeObject,
+						static_cast<uint8*>(DefaultValue),
+						DataID);
+				continue;
 			}
-			return true;
 		}
 	}
-	return false;
+
+	// Add any custom parameters the node may declare on its own.
+	if (ICustomCameraNodeParameterProvider* CustomParameterProvider = Cast<ICustomCameraNodeParameterProvider>(InCameraNode))
+	{
+		CustomParameterProvider->GetCustomCameraNodeParameters(*this);
+	}
 }
 
+}  // namespace UE::Cameras
+
 void ICustomCameraNodeParameterProvider::OnCustomCameraNodeParametersChanged(const UCameraNode* ThisAsCameraNode) const
 {
 	using namespace UE::Cameras;
diff --git a/Engine/Plugins/Cameras/GameplayCameras/Source/GameplayCameras/Private/GameFramework/BlueprintCameraEvaluationDataRef.cpp b/Engine/Plugins/Cameras/GameplayCameras/Source/GameplayCameras/Private/GameFramework/BlueprintCameraEvaluationDataRef.cpp
--- a/Engine/Plugins/Cameras/GameplayCameras/Source/GameplayCameras/Private/GameFramework/BlueprintCameraEvaluationDataRef.cpp
+++ b/Engine/Plugins/Cameras/GameplayCameras/Source/GameplayCameras/Private/GameFramework/BlueprintCameraEvaluationDataRef.cpp
@@ -73,7 +73,8 @@ void UBlueprintCameraEvaluationDataFunctionLibrary::SetDefaultCameraRigParameter
 
 	if (FCameraNodeEvaluationResult* Result = CameraData.GetResult())
 	{
-		FCameraObjectInterfaceParameterOverrideHelper::ApplyDefaultParameters(CameraRig, Result->VariableTable, Result->ContextDataTable);
+		FCameraObjectInterfaceParameterOverrideHelper Helper(*Result);
+		Helper.ApplyParameterDefaults(CameraRig);
 	}
 }
 
diff --git a/Engine/Plugins/Cameras/GameplayCameras/Source/GameplayCameras/Private/GameFramework/GameplayCameraComponent.cpp b/Engine/Plugins/Cameras/GameplayCameras/Source/GameplayCameras/Private/GameFramework/GameplayCameraComponent.cpp
--- a/Engine/Plugins/Cameras/GameplayCameras/Source/GameplayCameras/Private/GameFramework/GameplayCameraComponent.cpp
+++ b/Engine/Plugins/Cameras/GameplayCameras/Source/GameplayCameras/Private/GameFramework/GameplayCameraComponent.cpp
@@ -119,10 +119,17 @@ void UGameplayCameraComponent::ApplyChangedParameterOverrides()
 	{
 		FCameraNodeEvaluationResult& InitialResult = GetEvaluationContext()->GetInitialResult();
 
-		FCameraObjectInterfaceParameterOverrideHelper Helper(&InitialResult.VariableTable, &InitialResult.ContextDataTable);
-		Helper.bOverridenOnly = false;
-		GetChangedParameterOverrides(CameraReference.GetParameters(), CachedParameterOverrides, Helper.OnlyParameterGuids);
-		Helper.ApplyParameterOverrides(CameraAsset, CameraAsset->GetParameterDefinitions(), CameraReference.GetParameters());
+		TArray<FGuid> ChangedParameterGuids;
+		FCameraObjectInterfaceParameterOverrideHelper Helper(InitialResult);
+		GetChangedParameterOverrides(CameraReference.GetParameters(), CachedParameterOverrides, ChangedParameterGuids);
+		Helper.ApplyFilteredParameters(
+				CameraAsset,
+				CameraAsset->GetParameterDefinitions(),
+				CameraReference.GetParameters(),
+				[ChangedParameterGuids](const FCameraObjectInterfaceParameterDefinition& Definition) -> bool
+				{
+					return ChangedParameterGuids.Contains(Definition.ParameterGuid);
+				});
 
 		CachedParameterOverrides = CameraReference.GetParameters();
 	}
diff --git a/Engine/Plugins/Cameras/GameplayCameras/Source/GameplayCameras/Private/GameFramework/GameplayCameraRigComponent.cpp b/Engine/Plugins/Cameras/GameplayCameras/Source/GameplayCameras/Private/GameFramework/GameplayCameraRigComponent.cpp
--- a/Engine/Plugins/Cameras/GameplayCameras/Source/GameplayCameras/Private/GameFramework/GameplayCameraRigComponent.cpp
+++ b/Engine/Plugins/Cameras/GameplayCameras/Source/GameplayCameras/Private/GameFramework/GameplayCameraRigComponent.cpp
@@ -165,10 +165,17 @@ void UGameplayCameraRigComponent::ApplyChangedParameterOverrides()
 	{
 		FCameraNodeEvaluationResult& InitialResult = GetEvaluationContext()->GetInitialResult();
 
-		FCameraObjectInterfaceParameterOverrideHelper Helper(&InitialResult.VariableTable, &InitialResult.ContextDataTable);
-		Helper.bOverridenOnly = false;
-		GetChangedParameterOverrides(CameraRigReference.GetParameters(), CachedParameterOverrides, Helper.OnlyParameterGuids);
-		Helper.ApplyParameterOverrides(CameraRigAsset, CameraRigAsset->GetParameterDefinitions(), CameraRigReference.GetParameters());
+		TArray<FGuid> ChangedParameterGuids;
+		FCameraObjectInterfaceParameterOverrideHelper Helper(InitialResult);
+		GetChangedParameterOverrides(CameraRigReference.GetParameters(), CachedParameterOverrides, ChangedParameterGuids);
+		Helper.ApplyFilteredParameters(
+				CameraRigAsset,
+				CameraRigAsset->GetParameterDefinitions(),
+				CameraRigReference.GetParameters(),
+				[ChangedParameterGuids](const FCameraObjectInterfaceParameterDefinition& Definition) -> bool
+				{
+					return ChangedParameterGuids.Contains(Definition.ParameterGuid);
+				});
 
 		CachedParameterOverrides = CameraRigReference.GetParameters();
 	}
diff --git a/Engine/Plugins/Cameras/GameplayCameras/Source/GameplayCameras/Private/Helpers/CameraObjectInterfaceParameterOverrideHelper.cpp b/Engine/Plugins/Cameras/GameplayCameras/Source/GameplayCameras/Private/Helpers/CameraObjectInterfaceParameterOverrideHelper.cpp
--- a/Engine/Plugins/Cameras/GameplayCameras/Source/GameplayCameras/Private/Helpers/CameraObjectInterfaceParameterOverrideHelper.cpp
+++ b/Engine/Plugins/Cameras/GameplayCameras/Source/GameplayCameras/Private/Helpers/CameraObjectInterfaceParameterOverrideHelper.cpp
@@ -4,11 +4,13 @@
 
 #include "Core/BaseCameraObject.h"
 #include "Core/CameraContextDataTable.h"
+#include "Core/CameraNodeEvaluator.h"
 #include "Core/CameraObjectInterfaceParameterDefinition.h"
 #include "Core/CameraParameters.h"
 #include "Core/CameraVariableTable.h"
 #include "Misc/EngineVersionComparison.h"
 #include "StructUtils/OverridablePropertyBag.h"
+#include "Templates/Function.h"
 
 namespace UE::Cameras
 {
@@ -54,7 +56,7 @@ void ApplyBlendableParameterOverride(
 void ApplyBlendableParameterOverride(
 		const UObject* CameraObject,
 		const FCameraObjectInterfaceParameterDefinition& ParameterDefinition,
-		const FInstancedOverridablePropertyBag& PropertyBag,
+		const FInstancedPropertyBag& PropertyBag,
 		const FPropertyBagPropertyDesc& PropertyBagPropertyDesc,
 		FCameraVariableTable& VariableTable,
 		bool bDrivenOnly)
@@ -200,7 +202,7 @@ void ApplyDataParameterElementOverride(
 void ApplyDataParameterSingleOverride(
 		const UObject* CameraObject,
 		const FCameraObjectInterfaceParameterDefinition& ParameterDefinition,
-		const FInstancedOverridablePropertyBag& PropertyBag,
+		const FInstancedPropertyBag& PropertyBag,
 		const FPropertyBagPropertyDesc& PropertyBagPropertyDesc,
 		FCameraContextDataTable& ContextDataTable)
 {
@@ -283,7 +285,7 @@ void ApplyDataParameterSingleOverride(
 void ApplyDataParameterArrayOverride(
 		const UObject* CameraObject,
 		const FCameraObjectInterfaceParameterDefinition& ParameterDefinition,
-		const FInstancedOverridablePropertyBag& PropertyBag,
+		const FInstancedPropertyBag& PropertyBag,
 		const FPropertyBagPropertyDesc& PropertyBagPropertyDesc,
 		FCameraContextDataTable& ContextDataTable)
 {
@@ -370,7 +372,7 @@ void ApplyDataParameterArrayOverride(
 void ApplyDataParameterOverride(
 		const UObject* CameraObject,
 		const FCameraObjectInterfaceParameterDefinition& ParameterDefinition,
-		const FInstancedOverridablePropertyBag& PropertyBag,
+		const FInstancedPropertyBag& PropertyBag,
 		const FPropertyBagPropertyDesc& PropertyBagPropertyDesc,
 		FCameraContextDataTable& ContextDataTable)
 {
@@ -406,10 +408,17 @@ FCameraObjectInterfaceParameterOverrideHelper::FCameraObjectInterfaceParameterOv
 {
 }
 
-void FCameraObjectInterfaceParameterOverrideHelper::ApplyParameterOverrides(
-		const UObject* CameraObject,
-		TConstArrayView<FCameraObjectInterfaceParameterDefinition> ParameterDefinitions,
-		const FInstancedOverridablePropertyBag& ParameterOverrides)
+FCameraObjectInterfaceParameterOverrideHelper::FCameraObjectInterfaceParameterOverrideHelper(FCameraNodeEvaluationResult& OutResult)
+	: VariableTable(&OutResult.VariableTable)
+	, ContextDataTable(&OutResult.ContextDataTable)
+{
+}
+
+void FCameraObjectInterfaceParameterOverrideHelper::ApplyFilteredParameters(
+	const UObject* CameraObject,
+	TConstArrayView<FCameraObjectInterfaceParameterDefinition> ParameterDefinitions,
+	const FInstancedPropertyBag& ParameterOverrides,
+	TFunctionRef<bool(const FCameraObjectInterfaceParameterDefinition&)> ParameterFilter)
 {
 	using namespace Internal;
 
@@ -428,12 +437,7 @@ void FCameraObjectInterfaceParameterOverrideHelper::ApplyParameterOverrides(
 			continue;
 		}
 
-		if (bOverridenOnly && !ParameterOverrides.IsPropertyOverriden(Definition.ParameterGuid))
-		{
-			continue;
-		}
-
-		if (!OnlyParameterGuids.IsEmpty() && !OnlyParameterGuids.Contains(Definition.ParameterGuid))
+		if (!ParameterFilter(Definition))
 		{
 			continue;
 		}
@@ -444,14 +448,69 @@ void FCameraObjectInterfaceParameterOverrideHelper::ApplyParameterOverrides(
 			continue;
 		}
 		
-		ApplyParameterOverride(CameraObject, Definition, ParameterOverrides, *PropertyDesc);
+		ApplyParameterValue(CameraObject, Definition, ParameterOverrides, *PropertyDesc);
 	}
 }
 
-void FCameraObjectInterfaceParameterOverrideHelper::ApplyParameterOverride(
+void FCameraObjectInterfaceParameterOverrideHelper::ApplyParameters(
+		const UBaseCameraObject* CameraObject,
+		const FInstancedPropertyBag& Parameters)
+{
+	// Apply all parameters in the property bag.
+	ApplyFilteredParameters(
+			CameraObject,
+			CameraObject->GetParameterDefinitions(),
+			Parameters,
+			[](const FCameraObjectInterfaceParameterDefinition& Definition) -> bool
+			{
+				return true;
+			});
+}
+
+void FCameraObjectInterfaceParameterOverrideHelper::ApplyParameterOverrides(
+		const UBaseCameraObject* CameraObject,
+		const FInstancedOverridablePropertyBag& ParameterOverrides)
+{
+	// Apply parameters that have been overriden in the property bag.
+	ApplyFilteredParameters(
+			CameraObject,
+			CameraObject->GetParameterDefinitions(),
+			ParameterOverrides,
+			[ParameterOverrides](const FCameraObjectInterfaceParameterDefinition& Definition) -> bool
+			{
+				return ParameterOverrides.IsPropertyOverriden(Definition.ParameterGuid);
+			});
+}
+
+void FCameraObjectInterfaceParameterOverrideHelper::ApplyParameterDefaults(const UBaseCameraObject* CameraObject, bool bUnwrittenOnly)
+{
+	// Apply all default parameters, but only if they haven't been set in the variable/context data tables.
+	ApplyFilteredParameters(
+			CameraObject,
+			CameraObject->GetParameterDefinitions(),
+			CameraObject->GetDefaultParameters(),
+			[this, bUnwrittenOnly](const FCameraObjectInterfaceParameterDefinition& Definition) -> bool
+			{
+				if (bUnwrittenOnly)
+				{
+					if (Definition.ParameterType == ECameraObjectInterfaceParameterType::Blendable)
+					{
+						return Definition.VariableID.IsValid() && !VariableTable->IsValueWritten(Definition.VariableID);
+					}
+					else if (Definition.ParameterType == ECameraObjectInterfaceParameterType::Data)
+					{
+						return Definition.DataID.IsValid() && !ContextDataTable->IsValueWritten(Definition.DataID);
+					}
+					return false;
+				}
+				return true;
+			});
+}
+
+void FCameraObjectInterfaceParameterOverrideHelper::ApplyParameterValue(
 		const UObject* CameraObject,
 		const FCameraObjectInterfaceParameterDefinition& ParameterDefinition,
-		const FInstancedOverridablePropertyBag& PropertyBag,
+		const FInstancedPropertyBag& PropertyBag,
 		const FPropertyBagPropertyDesc& PropertyBagPropertyDesc)
 {
 	using namespace Internal;
@@ -488,76 +547,5 @@ void FCameraObjectInterfaceParameterOverrideHelper::ApplyParameterOverride(
 	}
 }
 
-void FCameraObjectInterfaceParameterOverrideHelper::ApplyDefaultBlendableParameters(const UBaseCameraObject* CameraObject, FCameraVariableTable& OutVariableTable)
-{
-	ApplyDefaultParametersImpl(CameraObject, &OutVariableTable, nullptr);
-}
-
-void FCameraObjectInterfaceParameterOverrideHelper::ApplyDefaultParameters(const UBaseCameraObject* CameraObject, FCameraVariableTable& OutVariableTable, FCameraContextDataTable& OutContextDataTable)
-{
-	ApplyDefaultParametersImpl(CameraObject, &OutVariableTable, &OutContextDataTable);
-}
-
-void FCameraObjectInterfaceParameterOverrideHelper::ApplyDefaultParametersImpl(const UBaseCameraObject* CameraObject, FCameraVariableTable* OutVariableTable, FCameraContextDataTable* OutContextDataTable)
-{
-	if (!ensure(CameraObject))
-	{
-		return;
-	}
-
-	const FInstancedPropertyBag& DefaultParameters = CameraObject->GetDefaultParameters();
-	const uint8* RawDefaultParametersContainer = DefaultParameters.GetValue().GetMemory();
-
-	for (const FCameraObjectInterfaceParameterDefinition& Definition : CameraObject->GetParameterDefinitions())
-	{
-		if (Definition.ParameterType == ECameraObjectInterfaceParameterType::Blendable)
-		{
-			if (!OutVariableTable || !Definition.VariableID.IsValid())
-			{
-				continue;
-			}
-			
-			if (OutVariableTable->IsValueWritten(Definition.VariableID))
-			{
-				continue;
-			}
-
-			const FPropertyBagPropertyDesc* PropertyDesc = DefaultParameters.FindPropertyDescByID(Definition.ParameterGuid);
-			if (!ensure(PropertyDesc && PropertyDesc->CachedProperty))
-			{
-				continue;
-			}
-
-			const void* RawValuePtr = PropertyDesc->CachedProperty->ContainerPtrToValuePtr<void>(RawDefaultParametersContainer);
-			OutVariableTable->TrySetValue(Definition.VariableID, Definition.VariableType, Definition.BlendableStructType, (const uint8*)RawValuePtr);
-		}
-		else if (Definition.ParameterType == ECameraObjectInterfaceParameterType::Data)
-		{
-			if (!OutContextDataTable || !Definition.DataID.IsValid())
-			{
-				continue;
-			}
-
-			if (OutContextDataTable->IsValueWritten(Definition.DataID))
-			{
-				continue;
-			}
-
-			const FPropertyBagPropertyDesc* PropertyDesc = DefaultParameters.FindPropertyDescByID(Definition.ParameterGuid);
-			if (!ensure(PropertyDesc && PropertyDesc->CachedProperty))
-			{
-				continue;
-			}
-
-			uint8* RawDestPtr = OutContextDataTable->TryGetMutableRawDataPtr(Definition.DataID, Definition.DataType, Definition.DataTypeObject);
-			if (RawDestPtr)
-			{
-				const void* RawValuePtr = PropertyDesc->CachedProperty->ContainerPtrToValuePtr<void>(RawDefaultParametersContainer);
-				PropertyDesc->CachedProperty->CopyCompleteValue(RawDestPtr, RawValuePtr);
-			}
-		}
-	}
-}
-
 }  // namespace UE::Cameras
 
diff --git a/Engine/Plugins/Cameras/GameplayCameras/Source/GameplayCameras/Private/Helpers/CameraObjectInterfaceParameterOverrideHelper.h b/Engine/Plugins/Cameras/GameplayCameras/Source/GameplayCameras/Private/Helpers/CameraObjectInterfaceParameterOverrideHelper.h
--- a/Engine/Plugins/Cameras/GameplayCameras/Source/GameplayCameras/Private/Helpers/CameraObjectInterfaceParameterOverrideHelper.h
+++ b/Engine/Plugins/Cameras/GameplayCameras/Source/GameplayCameras/Private/Helpers/CameraObjectInterfaceParameterOverrideHelper.h
@@ -4,6 +4,7 @@
 
 #include "Containers/ArrayView.h"
 #include "Misc/Guid.h"
+#include "Templates/FunctionFwd.h"
 
 class UBaseCameraObject;
 class UCameraAsset;
@@ -18,6 +19,7 @@ namespace UE::Cameras
 
 class FCameraContextDataTable;
 class FCameraVariableTable;
+struct FCameraNodeEvaluationResult;
 
 /**
  * A helper class for applying camera object interface parameter overrides from a property bag, 
@@ -26,18 +28,6 @@ class FCameraVariableTable;
  */
 struct FCameraObjectInterfaceParameterOverrideHelper
 {
-public:
-
-	/** Sets default values of blendable interface parameters in the given variable table. */
-	static void ApplyDefaultBlendableParameters(const UBaseCameraObject* CameraObject, FCameraVariableTable& OutVariableTable);
-
-	/** Sets default values of interface parameters in the given variable and context data tables. */
-	static void ApplyDefaultParameters(const UBaseCameraObject* CameraObject, FCameraVariableTable& OutVariableTable, FCameraContextDataTable& OutContextDataTable);
-
-private:
-
-	static void ApplyDefaultParametersImpl(const UBaseCameraObject* CameraObject, FCameraVariableTable* OutVariableTable, FCameraContextDataTable* OutContextDataTable);
-
 public:
 
 	/** 
@@ -48,28 +38,47 @@ public:
 	 */
 	FCameraObjectInterfaceParameterOverrideHelper(FCameraVariableTable* OutVariableTable, FCameraContextDataTable* OutContextDataTable);
 
+	/** Creates a new helper instance. */
+	FCameraObjectInterfaceParameterOverrideHelper(FCameraNodeEvaluationResult& OutResult);
+
+public:
+
+	/** Sets all values of interface parameters in the given variable and context data tables. */
+	void ApplyParameters(
+			const UBaseCameraObject* CameraObject,
+			const FInstancedPropertyBag& Parameters);
+
 	/** Sets overriden values of interface parameters in the given variable and context data tables. */
 	void ApplyParameterOverrides(
-			const UObject* CameraObject,
-			TConstArrayView<FCameraObjectInterfaceParameterDefinition> ParameterDefinitions,
+			const UBaseCameraObject* CameraObject,
 			const FInstancedOverridablePropertyBag& ParameterOverrides);
 
+	/** Sets all of the given object's default parameter values in the given variable and context data tables. */
+	void ApplyParameterDefaults(
+			const UBaseCameraObject* CameraObject,
+			bool bUnwrittenOnly = true);
+
+public:
+
+	/** Sets the parameter values for those that pass the given filter. */
+	void ApplyFilteredParameters(
+		const UObject* CameraObject,
+		TConstArrayView<FCameraObjectInterfaceParameterDefinition> ParameterDefinitions,
+		const FInstancedPropertyBag& ParameterOverrides,
+		TFunctionRef<bool(const FCameraObjectInterfaceParameterDefinition&)> ParameterFilter);
+
 private:
 
-	void ApplyParameterOverride(
+	void ApplyParameterValue(
 			const UObject* CameraObject,
 			const FCameraObjectInterfaceParameterDefinition& ParameterDefinition,
-			const FInstancedOverridablePropertyBag& PropertyBag,
+			const FInstancedPropertyBag& PropertyBag,
 			const FPropertyBagPropertyDesc& PropertyBagPropertyDesc);
 
 public:
 
-	/** Only parameters flagged as "overriden" should be applied. */
-	bool bOverridenOnly = true;
 	/** Only parameters that are dynamicaly driven should be applied. */
 	bool bDrivenOnly = false;
-	/** If not empty, the list of parameters we should only apply. */
-	TArray<FGuid> OnlyParameterGuids;
 
 private:
 
diff --git a/Engine/Plugins/Cameras/GameplayCameras/Source/GameplayCameras/Private/Nodes/Common/CameraRigCameraNode.cpp b/Engine/Plugins/Cameras/GameplayCameras/Source/GameplayCameras/Private/Nodes/Common/CameraRigCameraNode.cpp
--- a/Engine/Plugins/Cameras/GameplayCameras/Source/GameplayCameras/Private/Nodes/Common/CameraRigCameraNode.cpp
+++ b/Engine/Plugins/Cameras/GameplayCameras/Source/GameplayCameras/Private/Nodes/Common/CameraRigCameraNode.cpp
@@ -44,22 +44,16 @@ void FCameraRigCameraNodeEvaluator::OnBuild(const FCameraNodeEvaluatorBuildParam
 
 void FCameraRigCameraNodeEvaluator::OnInitialize(const FCameraNodeEvaluatorInitializeParams& Params, FCameraNodeEvaluationResult& OutResult)
 {
-	// Apply overrides right away.
-	ApplyParameterOverrides(OutResult.VariableTable, OutResult.ContextDataTable, false);
-
+	// Apply overrides and defaults right away.
 	const UCameraRigCameraNode* PrefabNode = GetCameraNodeAs<UCameraRigCameraNode>();
-	if (const UCameraRigAsset* CameraRig = PrefabNode->CameraRigReference.GetCameraRig())
-	{
-		// Set default values for unset entries in the variable table, so that pre-blending from default 
-		// values works.
-		FCameraObjectInterfaceParameterOverrideHelper::ApplyDefaultBlendableParameters(CameraRig, OutResult.VariableTable);
-	}
+	PrefabNode->CameraRigReference.ApplyParameterOverridesAndDefaults(OutResult);
 }
 
 void FCameraRigCameraNodeEvaluator::OnUpdateParameters(const FCameraBlendedParameterUpdateParams& Params, FCameraBlendedParameterUpdateResult& OutResult)
 {
 	// Keep applying overrides in case they are driven by a variable.
-	ApplyParameterOverrides(OutResult.VariableTable, false);
+	const UCameraRigCameraNode* PrefabNode = GetCameraNodeAs<UCameraRigCameraNode>();
+	PrefabNode->CameraRigReference.ApplyParameterOverrides(OutResult.VariableTable, false);
 }
 
 void FCameraRigCameraNodeEvaluator::OnRun(const FCameraNodeEvaluationParams& Params, FCameraNodeEvaluationResult& OutResult)
@@ -70,18 +64,6 @@ void FCameraRigCameraNodeEvaluator::OnRun(const FCameraNodeEvaluationParams& Par
 	}
 }
 
-void FCameraRigCameraNodeEvaluator::ApplyParameterOverrides(FCameraVariableTable& OutVariableTable, bool bDrivenOnly)
-{
-	const UCameraRigCameraNode* PrefabNode = GetCameraNodeAs<UCameraRigCameraNode>();
-	PrefabNode->CameraRigReference.ApplyParameterOverrides(OutVariableTable, bDrivenOnly);
-}
-
-void FCameraRigCameraNodeEvaluator::ApplyParameterOverrides(FCameraVariableTable& OutVariableTable, FCameraContextDataTable& OutContextDataTable, bool bDrivenOnly)
-{
-	const UCameraRigCameraNode* PrefabNode = GetCameraNodeAs<UCameraRigCameraNode>();
-	PrefabNode->CameraRigReference.ApplyParameterOverrides(OutVariableTable, OutContextDataTable, bDrivenOnly);
-}
-
 }  // namespace UE::Cameras
 
 void UCameraRigCameraNode::OnPreBuild(FCameraBuildLog& BuildLog)
@@ -113,6 +95,12 @@ void UCameraRigCameraNode::OnBuild(FCameraObjectBuildContext& BuildContext)
 		return;
 	}
 
+	const bool bDidModify = CameraRigReference.PostBuild();
+	if (bDidModify)
+	{
+		Modify();
+	}
+
 	const UCameraRigAsset* OuterCameraRig = GetTypedOuter<UCameraRigAsset>();
 	if (OuterCameraRig != CameraRig)
 	{
@@ -141,7 +129,7 @@ void UCameraRigCameraNode::GatherPackages(FCameraRigPackages& OutPackages) const
 	}
 }
 
-void UCameraRigCameraNode::GetCustomCameraNodeParameters(FCustomCameraNodeParameterInfos& OutParameterInfos)
+void UCameraRigCameraNode::GetCustomCameraNodeParameters(FCameraNodeParameterInfos& OutParameterInfos)
 {
 	CameraRigReference.GetCustomCameraNodeParameters(OutParameterInfos);
 }
diff --git a/Engine/Plugins/Cameras/GameplayCameras/Source/GameplayCameras/Private/Nodes/Framing/BaseFramingCameraNode.cpp b/Engine/Plugins/Cameras/GameplayCameras/Source/GameplayCameras/Private/Nodes/Framing/BaseFramingCameraNode.cpp
--- a/Engine/Plugins/Cameras/GameplayCameras/Source/GameplayCameras/Private/Nodes/Framing/BaseFramingCameraNode.cpp
+++ b/Engine/Plugins/Cameras/GameplayCameras/Source/GameplayCameras/Private/Nodes/Framing/BaseFramingCameraNode.cpp
@@ -848,7 +848,7 @@ void UBaseFramingCameraNode::PostLoad()
 	}
 }
 
-void UBaseFramingCameraNode::GetCustomCameraNodeParameters(FCustomCameraNodeParameterInfos& OutParameterInfos)
+void UBaseFramingCameraNode::GetCustomCameraNodeParameters(FCameraNodeParameterInfos& OutParameterInfos)
 {
 	OutParameterInfos.AddBlendableParameter(
 			GET_MEMBER_NAME_CHECKED(UBaseFramingCameraNode, DeadZone),
diff --git a/Engine/Plugins/Cameras/GameplayCameras/Source/GameplayCameras/Private/Nodes/Shakes/CameraShakeCameraNode.cpp b/Engine/Plugins/Cameras/GameplayCameras/Source/GameplayCameras/Private/Nodes/Shakes/CameraShakeCameraNode.cpp
--- a/Engine/Plugins/Cameras/GameplayCameras/Source/GameplayCameras/Private/Nodes/Shakes/CameraShakeCameraNode.cpp
+++ b/Engine/Plugins/Cameras/GameplayCameras/Source/GameplayCameras/Private/Nodes/Shakes/CameraShakeCameraNode.cpp
@@ -33,11 +33,6 @@ protected:
 	virtual void OnBuild(const FCameraNodeEvaluatorBuildParams& Params) override;
 	virtual void OnRun(const FCameraNodeEvaluationParams& Params, FCameraNodeEvaluationResult& OutResult) override;
 
-private:
-
-	void ApplyParameterOverrides(FCameraVariableTable& OutVariableTable, bool bDrivenOnly);
-	void ApplyParameterOverrides(FCameraVariableTable& OutVariableTable, FCameraContextDataTable& OutContextDataTable, bool bDrivenOnly);
-
 private:
 
 	ECameraShakeEvaluationMode EvaluationMode = ECameraShakeEvaluationMode::VisualLayer;
@@ -77,8 +72,9 @@ void FCameraShakeCameraNodeEvaluator::OnBuild(const FCameraNodeEvaluatorBuildPar
 
 void FCameraShakeCameraNodeEvaluator::OnInitialize(const FCameraNodeEvaluatorInitializeParams& Params, FCameraNodeEvaluationResult& OutResult)
 {
-	// Apply overrides right away.
-	ApplyParameterOverrides(OutResult.VariableTable, OutResult.ContextDataTable, false);
+	// Apply overrides and defaults right away.
+	const UCameraShakeCameraNode* PrefabNode = GetCameraNodeAs<UCameraShakeCameraNode>();
+	PrefabNode->CameraShakeReference.ApplyParameterOverridesAndDefaults(OutResult);
 
 	// If evaluating the shake later in the visual layer, acquire the shake service we will use to
 	// keep that shake alive.
@@ -92,7 +88,8 @@ void FCameraShakeCameraNodeEvaluator::OnInitialize(const FCameraNodeEvaluatorIni
 void FCameraShakeCameraNodeEvaluator::OnUpdateParameters(const FCameraBlendedParameterUpdateParams& Params, FCameraBlendedParameterUpdateResult& OutResult)
 {
 	// Keep applying overrides in case they are driven by a variable.
-	ApplyParameterOverrides(OutResult.VariableTable, false);
+	const UCameraShakeCameraNode* PrefabNode = GetCameraNodeAs<UCameraShakeCameraNode>();
+	PrefabNode->CameraShakeReference.ApplyParameterOverrides(OutResult.VariableTable, false);
 }
 
 void FCameraShakeCameraNodeEvaluator::OnRun(const FCameraNodeEvaluationParams& Params, FCameraNodeEvaluationResult& OutResult)
@@ -121,18 +118,6 @@ void FCameraShakeCameraNodeEvaluator::OnRun(const FCameraNodeEvaluationParams& P
 	}
 }
 
-void FCameraShakeCameraNodeEvaluator::ApplyParameterOverrides(FCameraVariableTable& OutVariableTable, bool bDrivenOnly)
-{
-	const UCameraShakeCameraNode* PrefabNode = GetCameraNodeAs<UCameraShakeCameraNode>();
-	PrefabNode->CameraShakeReference.ApplyParameterOverrides(OutVariableTable, bDrivenOnly);
-}
-
-void FCameraShakeCameraNodeEvaluator::ApplyParameterOverrides(FCameraVariableTable& OutVariableTable, FCameraContextDataTable& OutContextDataTable, bool bDrivenOnly)
-{
-	const UCameraShakeCameraNode* PrefabNode = GetCameraNodeAs<UCameraShakeCameraNode>();
-	PrefabNode->CameraShakeReference.ApplyParameterOverrides(OutVariableTable, OutContextDataTable, bDrivenOnly);
-}
-
 }  // namespace UE::Cameras
 
 void UCameraShakeCameraNode::OnPreBuild(FCameraBuildLog& BuildLog)
@@ -160,6 +145,12 @@ void UCameraShakeCameraNode::OnBuild(FCameraObjectBuildContext& BuildContext)
 		return;
 	}
 
+	const bool bDidModify = CameraShakeReference.PostBuild();
+	if (bDidModify)
+	{
+		Modify();
+	}
+
 	// Whatever allocations our inner camera shake needs for its evaluators and
 	// their camera variables, we add that to our camera shake's allocation info.
 	// If we're going to be running the shake in a deferred way, however, we need
@@ -175,7 +166,7 @@ void UCameraShakeCameraNode::OnBuild(FCameraObjectBuildContext& BuildContext)
 	BuildContext.AllocationInfo.Append(CameraShakeAllocationInfo);
 }
 
-void UCameraShakeCameraNode::GetCustomCameraNodeParameters(FCustomCameraNodeParameterInfos& OutParameterInfos)
+void UCameraShakeCameraNode::GetCustomCameraNodeParameters(FCameraNodeParameterInfos& OutParameterInfos)
 {
 	CameraShakeReference.GetCustomCameraNodeParameters(OutParameterInfos);
 }
diff --git a/Engine/Plugins/Cameras/GameplayCameras/Source/GameplayCameras/Private/Nodes/Utility/BlueprintCameraNode.cpp b/Engine/Plugins/Cameras/GameplayCameras/Source/GameplayCameras/Private/Nodes/Utility/BlueprintCameraNode.cpp
--- a/Engine/Plugins/Cameras/GameplayCameras/Source/GameplayCameras/Private/Nodes/Utility/BlueprintCameraNode.cpp
+++ b/Engine/Plugins/Cameras/GameplayCameras/Source/GameplayCameras/Private/Nodes/Utility/BlueprintCameraNode.cpp
@@ -126,7 +126,7 @@ void FBlueprintCameraNodeEvaluator::ApplyParameterOverrides(const FCameraVariabl
 			{
 				const uint8* ValuePtr = VariableTable.TryGetValue(
 						BlendableParameter.OverrideVariableID,
-						BlendableParameter.ParameterType,
+						BlendableParameter.VariableType,
 						BlendableParameter.BlendableStructType);
 				if (ValuePtr)
 				{
@@ -146,8 +146,8 @@ void FBlueprintCameraNodeEvaluator::ApplyParameterOverrides(const FCameraVariabl
 			{
 				const uint8* DataPtr = ContextDataTable.TryGetData(
 						DataParameter.OverrideDataID,
-						DataParameter.ParameterType,
-						DataParameter.ParameterTypeObject);
+						DataParameter.DataType,
+						DataParameter.DataTypeObject);
 				if (DataPtr)
 				{
 					Property->SetValue_InContainer(EvaluatorBlueprint, DataPtr);
@@ -277,7 +277,8 @@ void UBlueprintCameraNodeEvaluator::SetDefaultOwningCameraRigParameters(FBluepri
 	if (FCameraNodeEvaluationResult* Result = TargetCameraData.GetResult())
 	{
 		const UCameraRigAsset* OwningCameraRig = BlueprintNode->GetTypedOuter<UCameraRigAsset>();
-		FCameraObjectInterfaceParameterOverrideHelper::ApplyDefaultParameters(OwningCameraRig, Result->VariableTable, Result->ContextDataTable);
+		FCameraObjectInterfaceParameterOverrideHelper Helper(*Result);
+		Helper.ApplyParameterDefaults(OwningCameraRig);
 	}
 }
 
@@ -538,7 +539,7 @@ void UBlueprintCameraNode::RebuildOverrides()
 		{
 			FCustomCameraNodeBlendableParameter NewOverride;
 			NewOverride.ParameterName = Property->GetFName();
-			NewOverride.ParameterType = BlendablePropertyType;
+			NewOverride.VariableType = BlendablePropertyType;
 
 			// If this blendable parameter existed before and had an overriding variable set,
 			// preserve that override.
@@ -561,8 +562,8 @@ void UBlueprintCameraNode::RebuildOverrides()
 		{
 			FCustomCameraNodeDataParameter NewOverride;
 			NewOverride.ParameterName = Property->GetFName();
-			NewOverride.ParameterType = DataPropertyType;
-			NewOverride.ParameterTypeObject = DataPropertyTypeObject;
+			NewOverride.DataType = DataPropertyType;
+			NewOverride.DataTypeObject = DataPropertyTypeObject;
 
 			// If this data parameter existed before and had an override data ID set,
 			// preserve that override.
@@ -610,7 +611,7 @@ void UBlueprintCameraNode::OnBuild(FCameraObjectBuildContext& BuildContext)
 	}
 }
 
-void UBlueprintCameraNode::GetCustomCameraNodeParameters(FCustomCameraNodeParameterInfos& OutParameterInfos)
+void UBlueprintCameraNode::GetCustomCameraNodeParameters(FCameraNodeParameterInfos& OutParameterInfos)
 {
 	using namespace UE::Cameras::Internal;
 
diff --git a/Engine/Plugins/Cameras/GameplayCameras/Source/GameplayCameras/Private/StructUtils/OverridablePropertyBag.cpp b/Engine/Plugins/Cameras/GameplayCameras/Source/GameplayCameras/Private/StructUtils/OverridablePropertyBag.cpp
--- a/Engine/Plugins/Cameras/GameplayCameras/Source/GameplayCameras/Private/StructUtils/OverridablePropertyBag.cpp
+++ b/Engine/Plugins/Cameras/GameplayCameras/Source/GameplayCameras/Private/StructUtils/OverridablePropertyBag.cpp
@@ -42,6 +42,23 @@ void FInstancedOverridablePropertyBag::SetPropertyOverriden(const FGuid& InPrope
 	}
 }
 
+bool FInstancedOverridablePropertyBag::IsPropertyOverriden(const FName InPropertyName) const
+{
+	if (const FPropertyBagPropertyDesc* PropertyDesc = FindPropertyDescByName(InPropertyName))
+	{
+		return IsPropertyOverriden(PropertyDesc->ID);
+	}
+	return false;
+}
+
+void FInstancedOverridablePropertyBag::SetPropertyOverriden(const FName InPropertyName, bool bIsOverriden)
+{
+	if (const FPropertyBagPropertyDesc* PropertyDesc = FindPropertyDescByName(InPropertyName))
+	{
+		return SetPropertyOverriden(PropertyDesc->ID, bIsOverriden);
+	}
+}
+
 void FInstancedOverridablePropertyBag::MigrateToNewBagInstanceWithOverrides(const FInstancedPropertyBag& NewBagInstance)
 {
 	FInstancedPropertyBag::MigrateToNewBagInstanceWithOverrides(NewBagInstance, OverridenPropertyIDs);
diff --git a/Engine/Plugins/Cameras/GameplayCameras/Source/GameplayCameras/Private/Tests/CameraAssetBuilderTests.cpp b/Engine/Plugins/Cameras/GameplayCameras/Source/GameplayCameras/Private/Tests/CameraAssetBuilderTests.cpp
--- a/Engine/Plugins/Cameras/GameplayCameras/Source/GameplayCameras/Private/Tests/CameraAssetBuilderTests.cpp
+++ b/Engine/Plugins/Cameras/GameplayCameras/Source/GameplayCameras/Private/Tests/CameraAssetBuilderTests.cpp
@@ -44,17 +44,19 @@ bool FCameraAssetBuilderDefaultParametersTest::RunTest(const FString& Parameters
 
 	UCameraRigAsset* CameraRigOne = FCameraRigAssetTestBuilder(TEXT("One"))
 		.MakeRootNode<UOffsetCameraNode>().Named(TEXT("Offset"))
-			.SetParameter(&UOffsetCameraNode::TranslationOffset, FVector3d(10.0, 20.0, 0.0))
 			.Done()
 		.AddBlendableParameter(TEXT("TranslationOffset"), ECameraVariableType::Vector3d, TEXT("Offset"), GET_MEMBER_NAME_CHECKED(UOffsetCameraNode, TranslationOffset))
+		.BuildCameraRig()
+		.SetDefaultParameterValue<FVector3d>(TEXT("TranslationOffset"), FVector3d(10.0, 20.0, 0.0))
 		.Get();
 	CameraDirector->AddCameraRig(CameraRigOne, TEXT("One"));
 
 	UCameraRigAsset* CameraRigTwo = FCameraRigAssetTestBuilder(TEXT("Two"))
 		.MakeRootNode<ULensParametersCameraNode>().Named(TEXT("Lens"))
-			.SetParameter(&ULensParametersCameraNode::FocalLength, 55.0f)
 			.Done()
 		.AddBlendableParameter(TEXT("FocalLength"), ECameraVariableType::Float, TEXT("Lens"), GET_MEMBER_NAME_CHECKED(ULensParametersCameraNode, FocalLength))
+		.BuildCameraRig()
+		.SetDefaultParameterValue<float>(TEXT("FocalLength"), 55.0f)
 		.Get();
 	CameraDirector->AddCameraRig(CameraRigTwo, TEXT("Two"));
 
@@ -90,33 +92,36 @@ bool FCameraAssetBuilderDefaultParametersCollisionTest::RunTest(const FString& P
 
 	UCameraRigAsset* CameraRigOne = FCameraRigAssetTestBuilder(TEXT("One"))
 		.MakeRootNode<UOffsetCameraNode>().Named(TEXT("Offset"))
-			.SetParameter(&UOffsetCameraNode::TranslationOffset, FVector3d(10.0, 20.0, 0.0))
 			.Done()
 		.AddBlendableParameter(TEXT("TranslationOffset"), ECameraVariableType::Vector3d, TEXT("Offset"), GET_MEMBER_NAME_CHECKED(UOffsetCameraNode, TranslationOffset))
 		.AddBlendableParameter(TEXT("RotationOffset"), ECameraVariableType::Rotator3d, TEXT("Offset"), GET_MEMBER_NAME_CHECKED(UOffsetCameraNode, RotationOffset))
+		.BuildCameraRig()
+		.SetDefaultParameterValue<FVector3d>(TEXT("TranslationOffset"), FVector3d(10.0, 20.0, 0.0))
 		.Get();
 	CameraDirector->AddCameraRig(CameraRigOne, TEXT("One"));
 
 	UCameraRigAsset* CameraRigTwo = FCameraRigAssetTestBuilder(TEXT("Two"))
 		.MakeRootNode<ULensParametersCameraNode>().Named(TEXT("Lens"))
-			.SetParameter(&ULensParametersCameraNode::FocalLength, 55.0f)
 			.Done()
 		.AddBlendableParameter(TEXT("FocalLength"), ECameraVariableType::Float, TEXT("Lens"), GET_MEMBER_NAME_CHECKED(ULensParametersCameraNode, FocalLength))
 		.AddBlendableParameter(TEXT("Aperture"), ECameraVariableType::Float, TEXT("Lens"), GET_MEMBER_NAME_CHECKED(ULensParametersCameraNode, Aperture))
+		.BuildCameraRig()
+		.SetDefaultParameterValue<float>(TEXT("FocalLength"), 55.0f)
 		.Get();
 	CameraDirector->AddCameraRig(CameraRigTwo, TEXT("Two"));
 
 	UCameraRigAsset* CameraRigThree = FCameraRigAssetTestBuilder(TEXT("Three"))
 		.MakeArrayRootNode()
 			.AddChild<UOffsetCameraNode>(&UArrayCameraNode::Children).Named(TEXT("Offset2"))
-				.SetParameter(&UOffsetCameraNode::TranslationOffset, FVector3d(0.0, 10.0, 20.0))
 				.Done()
 			.AddChild<ULensParametersCameraNode>(&UArrayCameraNode::Children).Named(TEXT("Lens2"))
-				.SetParameter(&ULensParametersCameraNode::FocalLength, 35.0f)
 				.Done()
 			.Done()
 		.AddBlendableParameter(TEXT("TranslationOffset"), ECameraVariableType::Vector3d, TEXT("Offset2"), GET_MEMBER_NAME_CHECKED(UOffsetCameraNode, TranslationOffset))
 		.AddBlendableParameter(TEXT("FocalLength"), ECameraVariableType::Float, TEXT("Lens2"), GET_MEMBER_NAME_CHECKED(ULensParametersCameraNode, FocalLength))
+		.BuildCameraRig()
+		.SetDefaultParameterValue<FVector3d>(TEXT("TranslationOffset"), FVector3d(0.0, 10.0, 20.0))
+		.SetDefaultParameterValue<float>(TEXT("FocalLength"), 35.0f)
 		.Get();
 	CameraDirector->AddCameraRig(CameraRigThree, TEXT("Three"));
 
diff --git a/Engine/Plugins/Cameras/GameplayCameras/Source/GameplayCameras/Private/Tests/CameraPrefabTests.cpp b/Engine/Plugins/Cameras/GameplayCameras/Source/GameplayCameras/Private/Tests/CameraPrefabTests.cpp
new file mode 100644
--- /dev/null
+++ b/Engine/Plugins/Cameras/GameplayCameras/Source/GameplayCameras/Private/Tests/CameraPrefabTests.cpp
@@ -0,0 +1,160 @@
+// Copyright Epic Games, Inc. All Rights Reserved.
+
+#include "Core/CameraAsset.h"
+#include "Core/CameraEvaluationContext.h"
+#include "Core/CameraRigAsset.h"
+#include "Core/CameraSystemEvaluator.h"
+#include "Core/RootCameraNode.h"
+#include "Misc/AutomationTest.h"
+#include "Nodes/Common/CameraRigCameraNode.h"
+#include "Nodes/Common/LensParametersCameraNode.h"
+#include "Tests/GameplayCamerasTestBuilder.h"
+#include "Tests/GameplayCamerasTestObjects.h"
+
+#define LOCTEXT_NAMESPACE "CameraPrefabTests"
+
+namespace UE::Cameras::Test
+{
+
+TSharedRef<FNamedObjectRegistry> CreatePrefabTestCameraRigs()
+{
+	TSharedRef<FNamedObjectRegistry> NamedObjectRegistry = MakeShared<FNamedObjectRegistry>();
+
+	// A camera rig with a lens parameter node.
+	// That node has FocalLength, Aperture, and FocusDistance exposed as rig parameters.
+	// Default values are 22, 3, and 500 respectively.
+	UCameraRigAsset* InnerRig = FCameraRigAssetTestBuilder(NamedObjectRegistry, TEXT("InnerRig"))
+		.MakeRootNode<ULensParametersCameraNode>()
+			.Named(TEXT("InnerLensNode"))
+			.Done()
+		.AddBlendableParameter(
+				TEXT("InnerFocalLength"), ECameraVariableType::Float, 
+				TEXT("InnerLensNode"), TEXT("FocalLength"))
+		.AddBlendableParameter(
+				TEXT("InnerAperture"), ECameraVariableType::Float,
+				TEXT("InnerLensNode"), TEXT("Aperture"))
+		.AddBlendableParameter(
+				TEXT("InnerFocusDistance"), ECameraVariableType::Float,
+				TEXT("InnerLensNode"), TEXT("FocusDistance"))
+		.BuildCameraRig()
+		.SetDefaultParameterValue<float>(TEXT("InnerFocalLength"), 22.f)
+		.SetDefaultParameterValue<float>(TEXT("InnerAperture"), 3.f)
+		.SetDefaultParameterValue<float>(TEXT("InnerFocusDistance"), 500.f)
+		.Named(TEXT("InnerRig"))
+		.Get();
+
+	// A rig that uses the first one as a prefab.
+	// It forwards FocalLength and Aperture as rig parameters.
+	// It changes their default values to 18 and 4.5 respectively.
+	// It overrides FocusDistance to 600.
+	UCameraRigAsset* FirstLevelRig = FCameraRigAssetTestBuilder(NamedObjectRegistry, TEXT("FirstLevelRig"))
+		.MakeRootNode<UCameraRigCameraNode>()
+			.Named(TEXT("FirstLevelPrefabNode"))
+			.Setup([InnerRig](UCameraRigCameraNode* Node)
+				{
+					Node->CameraRigReference.SetCameraRig(InnerRig);
+				})
+			.Done()
+		.AddBlendableParameter(
+				TEXT("FirstLevelFocalLength"), ECameraVariableType::Float,
+				TEXT("FirstLevelPrefabNode"), TEXT("InnerFocalLength"))
+		.AddBlendableParameter(
+				TEXT("FirstLevelAperture"), ECameraVariableType::Float,
+				TEXT("FirstLevelPrefabNode"), TEXT("InnerAperture"))
+		.BuildCameraRig()
+		.SetDefaultParameterValue<float>(TEXT("FirstLevelFocalLength"), 18.f)
+		.SetDefaultParameterValue<float>(TEXT("FirstLevelAperture"), 4.5f)
+		.Setup([](UCameraRigAsset* CameraRig, FNamedObjectRegistry* NamedObjectRegistry)
+				{
+					UCameraRigCameraNode* FirstLevelPrefabNode = NamedObjectRegistry->Get<UCameraRigCameraNode>(TEXT("FirstLevelPrefabNode"));
+					FInstancedOverridablePropertyBag& Parameters = FirstLevelPrefabNode->CameraRigReference.GetParameters();
+					Parameters.SetPropertyOverriden(TEXT("InnerFocusDistance"), true);
+					Parameters.GetValueStruct<FFloatCameraParameter>(TEXT("InnerFocusDistance")).GetValue()->Value = 600.f;
+				})
+		.Named(TEXT("FirstLevelRig"))
+		.Get();
+
+	// A rig that uses the previous rig as a prefab.
+	// It overrides FocalLenth to 16.
+	FCameraRigAssetTestBuilder(NamedObjectRegistry, TEXT("SecondLevelRig"))
+		.MakeRootNode<UCameraRigCameraNode>()
+			.Named(TEXT("SecondLevelPrefabNode"))
+			.Setup([FirstLevelRig](UCameraRigCameraNode* Node)
+				{
+					Node->CameraRigReference.SetCameraRig(FirstLevelRig);
+				})
+			.Done()
+		.BuildCameraRig()
+		.Setup([](UCameraRigAsset* CameraRig, FNamedObjectRegistry* NamedObjectRegistry)
+				{
+					UCameraRigCameraNode* SecondLevelPrefabNode = NamedObjectRegistry->Get<UCameraRigCameraNode>(TEXT("SecondLevelPrefabNode"));
+					FInstancedOverridablePropertyBag& Parameters = SecondLevelPrefabNode->CameraRigReference.GetParameters();
+					Parameters.SetPropertyOverriden(TEXT("FirstLevelFocalLength"), true);
+					Parameters.GetValueStruct<FFloatCameraParameter>(TEXT("FirstLevelFocalLength")).GetValue()->Value = 16.f;
+				})
+		.Named(TEXT("SecondLevelRig"))
+		.Get();
+
+	return NamedObjectRegistry;
+}
+
+void RunCameraRig(UCameraRigAsset* InCameraRig, FCameraPose& OutCameraPose)
+{
+	TSharedRef<FCameraEvaluationContext> EvaluationContext = FCameraEvaluationContextTestBuilder()
+		.MakeSingleDirector(InCameraRig)
+			.Done()
+		.BuildCameraAsset()
+		.Get();
+
+	TSharedRef<FCameraSystemEvaluator> Evaluator = FCameraSystemEvaluatorBuilder::Build();
+	Evaluator->PushEvaluationContext(EvaluationContext);
+
+	EvaluationContext->GetInitialResult().bIsValid = true;
+
+	FCameraSystemEvaluationParams Params;
+	Params.DeltaTime = 0.3f;
+	Evaluator->Update(Params);
+
+	const FCameraSystemEvaluationResult& EvaluatedResult = Evaluator->GetEvaluatedResult();
+	OutCameraPose = EvaluatedResult.CameraPose;
+}
+
+}  // namespace UE::Cameras::Test
+
+IMPLEMENT_SIMPLE_AUTOMATION_TEST(FCameraRigPrefabDefaultsTest, "System.Engine.GameplayCameras.Prefab.Defaults", EAutomationTestFlags::EditorContext | EAutomationTestFlags::EngineFilter)
+bool FCameraRigPrefabDefaultsTest::RunTest(const FString& Parameters)
+{
+	using namespace UE::Cameras;
+	using namespace UE::Cameras::Test;
+
+	TSharedRef<FNamedObjectRegistry> NamedObjectRegistry = CreatePrefabTestCameraRigs();
+
+	{
+		FCameraPose CameraPose;
+		RunCameraRig(NamedObjectRegistry->Get<UCameraRigAsset>("InnerRig"), CameraPose);
+		UTEST_EQUAL(TEXT("FocalLength"), CameraPose.GetFocalLength(), 22.f);
+		UTEST_EQUAL(TEXT("Aperture"), CameraPose.GetAperture(), 3.f);
+		UTEST_EQUAL(TEXT("FocusDistance"), CameraPose.GetFocusDistance(), 500.f);
+	}
+
+	{
+		FCameraPose CameraPose;
+		RunCameraRig(NamedObjectRegistry->Get<UCameraRigAsset>("FirstLevelRig"), CameraPose);
+		UTEST_EQUAL(TEXT("FocalLength"), CameraPose.GetFocalLength(), 18.f);
+		UTEST_EQUAL(TEXT("Aperture"), CameraPose.GetAperture(), 4.5f);
+		UTEST_EQUAL(TEXT("FocusDistance"), CameraPose.GetFocusDistance(), 600.f);
+	}
+
+	{
+		FCameraPose CameraPose;
+		RunCameraRig(NamedObjectRegistry->Get<UCameraRigAsset>("SecondLevelRig"), CameraPose);
+		UTEST_EQUAL(TEXT("FocalLength"), CameraPose.GetFocalLength(), 16.f);
+		UTEST_EQUAL(TEXT("Aperture"), CameraPose.GetAperture(), 4.5f);
+		UTEST_EQUAL(TEXT("FocusDistance"), CameraPose.GetFocusDistance(), 600.f);
+	}
+
+	return true;
+}
+
+#undef LOCTEXT_NAMESPACE
+
diff --git a/Engine/Plugins/Cameras/GameplayCameras/Source/GameplayCameras/Private/Tests/CameraSystemUpdateTests.cpp b/Engine/Plugins/Cameras/GameplayCameras/Source/GameplayCameras/Private/Tests/CameraSystemUpdateTests.cpp
--- a/Engine/Plugins/Cameras/GameplayCameras/Source/GameplayCameras/Private/Tests/CameraSystemUpdateTests.cpp
+++ b/Engine/Plugins/Cameras/GameplayCameras/Source/GameplayCameras/Private/Tests/CameraSystemUpdateTests.cpp
@@ -16,7 +16,7 @@ bool FCameraSystemFrameFlagsTest::RunTest(const FString& Parameters)
 	using namespace UE::Cameras::Test;
 
 	TSharedRef<FCameraEvaluationContext> EvaluationContext = FCameraEvaluationContextTestBuilder()
-		.AddCameraRig(TEXT("TestRig"))
+		.CreateCameraRig(TEXT("TestRig"))
 			.MakeRootNode<UUpdateTrackerCameraNode>()
 				.Done()
 			.Done()
diff --git a/Engine/Plugins/Cameras/GameplayCameras/Source/GameplayCameras/Private/Tests/GameplayCamerasTestBuilder.h b/Engine/Plugins/Cameras/GameplayCameras/Source/GameplayCameras/Private/Tests/GameplayCamerasTestBuilder.h
--- a/Engine/Plugins/Cameras/GameplayCameras/Source/GameplayCameras/Private/Tests/GameplayCamerasTestBuilder.h
+++ b/Engine/Plugins/Cameras/GameplayCameras/Source/GameplayCameras/Private/Tests/GameplayCamerasTestBuilder.h
@@ -490,7 +490,7 @@ public:
 	{
 		UCameraObjectInterfaceBlendableParameter* BlendableParameter = NewObject<UCameraObjectInterfaceBlendableParameter>(CameraRig);
 		BlendableParameter->InterfaceParameterName = ParameterName;
-		BlendableParameter->ParameterType = ParameterType;
+		BlendableParameter->VariableType = ParameterType;
 
 		NamedObjectRegistry->Register(BlendableParameter, ParameterName);
 		CameraRig->Interface.BlendableParameters.Add(BlendableParameter);
@@ -515,6 +515,82 @@ public:
 		return AddBlendableParameter(ParameterName, ParameterType, Target, TargetPropertyName);
 	}
 
+	/** Runs arbitrary setup logic on the camera rig. */
+	ThisType& Setup(TFunction<void(UCameraRigAsset*)> SetupCallback)
+	{
+		SetupCallback(CameraRig);
+		return *static_cast<ThisType*>(this);;
+	}
+
+	/** Runs arbitrary setup logic on the camera rig. */
+	ThisType& Setup(TFunction<void(UCameraRigAsset*, FNamedObjectRegistry*)> SetupCallback)
+	{
+		SetupCallback(CameraRig, NamedObjectRegistry.Get());
+		return *static_cast<ThisType*>(this);;
+	}
+
+	/**
+	 * Builds the camera rig.
+	 */
+	ThisType& BuildCameraRig()
+	{
+		CameraRig->BuildCameraRig();
+		return *static_cast<ThisType*>(this);
+	}
+
+	/**
+	 * Sets the default value for a camera rig parameter.
+	 * Requires that the camera rig has been built first.
+	 */
+	template<typename ValueType>
+	ThisType& SetDefaultParameterValue(const FName ParameterName, TCallTraits<ValueType>::ParamType ParameterValue)
+	{
+		const FCameraObjectInterfaceParameterDefinition* ParameterDefinition = CameraRig->GetParameterDefinitions().FindByPredicate(
+				[ParameterName](const FCameraObjectInterfaceParameterDefinition& Item)
+				{
+					return Item.ParameterName == ParameterName;
+				});
+		if (!ensureMsgf(ParameterDefinition, TEXT("You must build the camera rig before setting default parameter values")))
+		{
+			return *static_cast<ThisType*>(this);
+		}
+
+		FInstancedPropertyBag& DefaultParameters = CameraRig->GetDefaultParameters();
+		const UPropertyBag* PropertyBag = DefaultParameters.GetPropertyBagStruct();
+		if (!ensureMsgf(PropertyBag, TEXT("You must build the camera rig before setting default parameter values")))
+		{
+			return *static_cast<ThisType*>(this);
+		}
+
+		const FPropertyBagPropertyDesc* PropertyDesc = PropertyBag->FindPropertyDescByName(ParameterName);
+		if (!ensureMsgf(PropertyDesc, TEXT("No such camera rig parameter")))
+		{
+			return *static_cast<ThisType*>(this);
+		}
+
+		uint8* RawParameters = DefaultParameters.GetMutableValue().GetMemory();
+		void* TargetAddress = RawParameters + PropertyDesc->CachedProperty->GetOffset_ForInternal();
+		const void* SourceAddress = &ParameterValue;
+
+		if (ParameterDefinition->ParameterType == ECameraObjectInterfaceParameterType::Blendable)
+		{
+			if (ParameterDefinition->VariableType == ECameraVariableType::BlendableStruct)
+			{
+				PropertyDesc->CachedProperty->CopyCompleteValue(TargetAddress, SourceAddress);
+			}
+			else
+			{
+				SetDefaultBlendableParameterValue(*ParameterDefinition, TargetAddress, ParameterValue);
+			}
+		}
+		else if (ParameterDefinition->ParameterType == ECameraObjectInterfaceParameterType::Data)
+		{
+			PropertyDesc->CachedProperty->CopyCompleteValue(TargetAddress, SourceAddress);
+		}
+
+		return *static_cast<ThisType*>(this);
+	}
+
 	/** Gets the named object registry. */
 	virtual TSharedPtr<FNamedObjectRegistry> GetNamedObjectRegistry() override
 	{
@@ -549,6 +625,23 @@ private:
 		NamedObjectRegistry->Register(CameraRig, Name.ToString());
 	}
 
+	template<typename ValueType>
+	void SetDefaultBlendableParameterValue(const FCameraObjectInterfaceParameterDefinition& ParameterDefinition, void* TargetAddress, ValueType ParameterValue);
+
+#define UE_CAMERA_VARIABLE_FOR_TYPE(ValueType, ValueName)\
+	template<>\
+	void SetDefaultBlendableParameterValue<ValueType>(const FCameraObjectInterfaceParameterDefinition& ParameterDefinition, void* TargetAddress, ValueType ParameterValue)\
+	{\
+		using ParameterType = F##ValueName##CameraParameter;\
+		if (ensure(ParameterDefinition.ParameterType == ECameraObjectInterfaceParameterType::Blendable && ParameterDefinition.VariableType == ECameraVariableType::ValueName))\
+		{\
+			ParameterType* TargetParameter = static_cast<ParameterType*>(TargetAddress);\
+			TargetParameter->Value = ParameterValue;\
+		}\
+	}
+	UE_CAMERA_VARIABLE_FOR_ALL_TYPES()
+#undef UE_CAMERA_VARIABLE_FOR_TYPE
+
 private:
 
 	UCameraRigAsset* CameraRig;
@@ -760,13 +853,52 @@ public:
 		return MakeDirector<USingleCameraDirector>();
 	}
 
-	/** Creates a new camera rig asset builder and adds its camera rig to the camera asset.*/
-	TScopedCameraRigAssetTestBuilder<ThisType> AddCameraRig(FName Name = NAME_None)
+	/** Builds a new single camera director, set its camera rig, and returns a builder object for its. */
+	TCameraDirectorTestBuilder<ThisType, USingleCameraDirector> MakeSingleDirector(UCameraRigAsset* InCameraRig)
+	{
+		auto DirectorBuilder = MakeDirector<USingleCameraDirector>();
+		DirectorBuilder.Setup([InCameraRig](USingleCameraDirector* Director)
+				{
+					Director->CameraRig = InCameraRig;
+				});
+		return DirectorBuilder;
+	}
+
+	/** Builds a new single camera director, set its camera rig to the named object, and returns a builder object for its. */
+	TCameraDirectorTestBuilder<ThisType, USingleCameraDirector> MakeSingleDirector(const FString& InCameraRigName)
+	{
+		UCameraRigAsset* CameraRig = NamedObjectRegistry->Get<UCameraRigAsset>(InCameraRigName);
+
+		auto DirectorBuilder = MakeDirector<USingleCameraDirector>();
+		DirectorBuilder.Setup([CameraRig](USingleCameraDirector* Director)
+				{
+					Director->CameraRig = CameraRig;
+				});
+		return DirectorBuilder;
+	}
+
+	/** Creates a new camera rig asset builder and gives it a name to be recalled later.*/
+	TScopedCameraRigAssetTestBuilder<ThisType> CreateCameraRig(FName Name = NAME_None)
 	{
 		TScopedCameraRigAssetTestBuilder<ThisType> CameraRigBuilder(*this, GetNamedObjectRegistry(), Name, CameraAsset);
 		return CameraRigBuilder;
 	}
 
+	/** Runs arbitrary setup logic on the evaluation context. */
+	ThisType& Setup(TFunction<void(TSharedRef<FCameraEvaluationContext>)> SetupCallback)
+	{
+		SetupCallback(EvaluationContext.ToSharedRef());
+		return *this;
+	}
+
+	/** Runs arbitrary setup logic on the evaluation context. */
+	ThisType& Setup(TFunction<void(TSharedRef<FCameraEvaluationContext>, FNamedObjectRegistry*)> SetupCallback)
+	{
+		SetupCallback(EvaluationContext.ToSharedRef(), NamedObjectRegistry.Get());
+		return *this;
+	}
+
+
 	/** Gets the named object registry. */
 	virtual TSharedPtr<FNamedObjectRegistry> GetNamedObjectRegistry() override
 	{
diff --git a/Engine/Plugins/Cameras/GameplayCameras/Source/GameplayCameras/Public/Build/CameraObjectInterfaceBuilder.h b/Engine/Plugins/Cameras/GameplayCameras/Source/GameplayCameras/Public/Build/CameraObjectInterfaceBuilder.h
--- a/Engine/Plugins/Cameras/GameplayCameras/Source/GameplayCameras/Public/Build/CameraObjectInterfaceBuilder.h
+++ b/Engine/Plugins/Cameras/GameplayCameras/Source/GameplayCameras/Public/Build/CameraObjectInterfaceBuilder.h
@@ -55,7 +55,7 @@ private:
 	UBaseCameraObject* CameraObject = nullptr;
 	TArray<UCameraNode*> CameraObjectNodes;
 
-	using FDrivenParameterKey = TTuple<FName, UObject*>;
+	using FDrivenParameterKey = TTuple<UObject*, FName>;
 	TMap<FDrivenParameterKey, FCameraVariableID> OldDrivenBlendableParameters;
 	TMap<FDrivenParameterKey, FCameraContextDataID> OldDrivenDataParameters;
 
diff --git a/Engine/Plugins/Cameras/GameplayCameras/Source/GameplayCameras/Public/Core/BaseCameraObjectReference.h b/Engine/Plugins/Cameras/GameplayCameras/Source/GameplayCameras/Public/Core/BaseCameraObjectReference.h
--- a/Engine/Plugins/Cameras/GameplayCameras/Source/GameplayCameras/Public/Core/BaseCameraObjectReference.h
+++ b/Engine/Plugins/Cameras/GameplayCameras/Source/GameplayCameras/Public/Core/BaseCameraObjectReference.h
@@ -9,11 +9,11 @@
 #include "BaseCameraObjectReference.generated.h"
 
 class UBaseCameraObject;
-struct FCustomCameraNodeParameterInfos;
 
 namespace UE::Cameras
 {
 	struct FCameraNodeEvaluationResult;
+	struct FCameraNodeParameterInfos;
 }
 
 /**
@@ -28,7 +28,11 @@ struct FCameraObjectInterfaceParameterMetaData
 	UPROPERTY()
 	FGuid ParameterGuid;
 
-	/** The ID to use for overriding a blendable parameter. */
+	/** 
+	 * The ID to use for overriding a blendable parameter.
+	 * NOTE: only used for blendable structs, since other types are parameter structs with their
+	 * own variable ID inside.
+	 */
 	UPROPERTY()
 	FCameraVariableID OverrideVariableID;
 
@@ -88,7 +92,9 @@ public:
 
 	void PostSerialize(const FArchive& Ar);
 
-	GAMEPLAYCAMERAS_API void GetCustomCameraNodeParameters(FCustomCameraNodeParameterInfos& OutParameterInfos);
+	GAMEPLAYCAMERAS_API void GetCustomCameraNodeParameters(UE::Cameras::FCameraNodeParameterInfos& OutParameterInfos);
+
+	bool PostBuild();
 	
 	GAMEPLAYCAMERAS_API bool IsParameterOverridden(const FGuid& PropertyID) const;
 	GAMEPLAYCAMERAS_API void SetParameterOverridden(const FGuid& PropertyID, bool bIsOverridden);
@@ -96,10 +102,13 @@ public:
 	template<typename ContainerType>
 	void GetOverriddenParameterGuids(ContainerType& OutOverriddenIDs) const;
 
-private:
+protected:
 
 	const FCameraObjectInterfaceParameterMetaData* FindMetaData(const FGuid& PropertyID) const;
-	FCameraObjectInterfaceParameterMetaData& FindOrAddMetaData(const FGuid& PropertyID);
+
+private:
+
+	FCameraObjectInterfaceParameterMetaData& GetMetaData(const FGuid& PropertyID);
 
 protected:
 
diff --git a/Engine/Plugins/Cameras/GameplayCameras/Source/GameplayCameras/Public/Core/CameraObjectInterface.h b/Engine/Plugins/Cameras/GameplayCameras/Source/GameplayCameras/Public/Core/CameraObjectInterface.h
--- a/Engine/Plugins/Cameras/GameplayCameras/Source/GameplayCameras/Public/Core/CameraObjectInterface.h
+++ b/Engine/Plugins/Cameras/GameplayCameras/Source/GameplayCameras/Public/Core/CameraObjectInterface.h
@@ -91,7 +91,7 @@ public:
 
 	/** The type of this parameter. */
 	UPROPERTY()
-	ECameraVariableType ParameterType = ECameraVariableType::Boolean;
+	ECameraVariableType VariableType = ECameraVariableType::Boolean;
 
 	/** The struct type of this parameter if it is a blendable struct. */
 	UPROPERTY()
diff --git a/Engine/Plugins/Cameras/GameplayCameras/Source/GameplayCameras/Public/Core/CameraRigAssetReference.h b/Engine/Plugins/Cameras/GameplayCameras/Source/GameplayCameras/Public/Core/CameraRigAssetReference.h
--- a/Engine/Plugins/Cameras/GameplayCameras/Source/GameplayCameras/Public/Core/CameraRigAssetReference.h
+++ b/Engine/Plugins/Cameras/GameplayCameras/Source/GameplayCameras/Public/Core/CameraRigAssetReference.h
@@ -288,9 +288,12 @@ public:
 	/** Applies the parameter override values to the given evaluation result. */
 	void ApplyParameterOverrides(UE::Cameras::FCameraNodeEvaluationResult& OutResult, bool bDrivenOnly) const;
 
+	/** Applies the parameter override values, or the default value if not overriden, to the given evaluation result. */
+	void ApplyParameterOverridesAndDefaults(UE::Cameras::FCameraNodeEvaluationResult& OutResult) const;
+
 private:
 
-	void ApplyParameterOverridesImpl(UE::Cameras::FCameraVariableTable* OutVariableTable, UE::Cameras::FCameraContextDataTable* OutContextDataTable, bool bDrivenOnly) const;
+	void ApplyParameterOverridesImpl(UE::Cameras::FCameraVariableTable* OutVariableTable, UE::Cameras::FCameraContextDataTable* OutContextDataTable, bool bDrivenOnly, bool bApplyUnwrittenDefaults) const;
 
 public:
 
diff --git a/Engine/Plugins/Cameras/GameplayCameras/Source/GameplayCameras/Public/Core/CameraShakeAssetReference.h b/Engine/Plugins/Cameras/GameplayCameras/Source/GameplayCameras/Public/Core/CameraShakeAssetReference.h
--- a/Engine/Plugins/Cameras/GameplayCameras/Source/GameplayCameras/Public/Core/CameraShakeAssetReference.h
+++ b/Engine/Plugins/Cameras/GameplayCameras/Source/GameplayCameras/Public/Core/CameraShakeAssetReference.h
@@ -70,9 +70,12 @@ public:
 	/** Applies the parameter override values to the given evaluation result. */
 	void ApplyParameterOverrides(UE::Cameras::FCameraNodeEvaluationResult& OutResult, bool bDrivenOnly) const;
 
+	/** Applies the parameter override values, or the default value if not overriden, to the given evaluation result. */
+	void ApplyParameterOverridesAndDefaults(UE::Cameras::FCameraNodeEvaluationResult& OutResult) const;
+
 private:
 
-	void ApplyParameterOverridesImpl(UE::Cameras::FCameraVariableTable* OutVariableTable, UE::Cameras::FCameraContextDataTable* OutContextDataTable, bool bDrivenOnly) const;
+	void ApplyParameterOverridesImpl(UE::Cameras::FCameraVariableTable* OutVariableTable, UE::Cameras::FCameraContextDataTable* OutContextDataTable, bool bDrivenOnly, bool bApplyUnwrittenDefaults) const;
 
 public:
 
diff --git a/Engine/Plugins/Cameras/GameplayCameras/Source/GameplayCameras/Public/Core/ICustomCameraNodeParameterProvider.h b/Engine/Plugins/Cameras/GameplayCameras/Source/GameplayCameras/Public/Core/ICustomCameraNodeParameterProvider.h
--- a/Engine/Plugins/Cameras/GameplayCameras/Source/GameplayCameras/Public/Core/ICustomCameraNodeParameterProvider.h
+++ b/Engine/Plugins/Cameras/GameplayCameras/Source/GameplayCameras/Public/Core/ICustomCameraNodeParameterProvider.h
@@ -13,11 +13,108 @@ class UCameraNode;
 
 namespace UE::Cameras
 {
-	class FCameraNodeHierarchyBuilder;
-	class FCameraObjectInterfaceBuilder;
-	class FCameraObjectInterfaceParameterBuilder;
-	namespace Internal { struct FInterfaceParameterBindingBuilder; }
-}
+
+class FCameraNodeHierarchyBuilder;
+class FCameraObjectInterfaceBuilder;
+class FCameraObjectInterfaceParameterBuilder;
+namespace Internal { struct FInterfaceParameterBindingBuilder; }
+
+/** Information about a blendable parameter on a camera node. */
+struct FCameraNodeBlendableParameterInfo
+{
+	FName ParameterName;
+	ECameraVariableType VariableType = ECameraVariableType::Boolean;
+	const UScriptStruct* BlendableStructType = nullptr;
+	const uint8* DefaultValue = nullptr;
+	FCameraVariableID* OverrideVariableID = nullptr;
+	UCameraVariableAsset* OverrideVariable = nullptr;
+};
+
+/** Information about a data parameter on a camera node. */
+struct FCameraNodeDataParameterInfo
+{
+	FName ParameterName;
+	ECameraContextDataType DataType = ECameraContextDataType::Name;
+	ECameraContextDataContainerType DataContainerType = ECameraContextDataContainerType::None;
+	const UObject* DataTypeObject = nullptr;
+	const uint8* DefaultValue = nullptr;
+	FCameraContextDataID* OverrideDataID = nullptr;
+};
+
+/**
+ * A structure providing information about the exposed parameters on a camera node.
+ */
+struct FCameraNodeParameterInfos
+{
+	using FBlendableParameterInfo = FCameraNodeBlendableParameterInfo;
+	using FDataParameterInfo = FCameraNodeDataParameterInfo;
+
+public:
+
+	/** 
+	 * Declares a blendable parameter. 
+	 * Pass a null pointer for OverrideVariableID if this parameter should not be overridable
+	 * by a camera rig parameter.
+	 */
+	GAMEPLAYCAMERAS_API void AddBlendableParameter(
+			FName ParameterName, 
+			ECameraVariableType VariableType, 
+			const UScriptStruct* BlendableStructType,
+			const uint8* DefaultValue,
+			FCameraVariableID* OverrideVariableID);
+
+	GAMEPLAYCAMERAS_API void AddBlendableParameter(FCustomCameraNodeBlendableParameter& Parameter, const uint8* DefaultValue);
+
+	/** 
+	 * Declares a data parameter.
+	 * Pass a null pointer for OverrideDataID if this parameter should not be overridable
+	 * by a camera rig parameter.
+	 */
+	GAMEPLAYCAMERAS_API void AddDataParameter(
+			FName ParameterName, 
+			ECameraContextDataType DataType,
+			ECameraContextDataContainerType DataContainerType,
+			const UObject* DataTypeObject,
+			const uint8* DefaultValue,
+			FCameraContextDataID* OverrideDataID);
+
+	GAMEPLAYCAMERAS_API void AddDataParameter(FCustomCameraNodeDataParameter& Parameter, const uint8* DefaultValue);
+
+	/** Clears the list of parameters in this structure. */
+	GAMEPLAYCAMERAS_API void Reset();
+
+public:
+
+	/** Returns whether there are any blendable or data parameters. */
+	bool HasAnyParameters() const { return !BlendableParameters.IsEmpty() || !DataParameters.IsEmpty(); }
+
+	/** Gets the list of blendable parameters. */
+	TConstArrayView<FBlendableParameterInfo> GetBlendableParameters() const { return BlendableParameters; }
+	/** Gets the list of data parameters. */
+	TConstArrayView<FDataParameterInfo> GetDataParameters() const { return DataParameters; }
+
+	/** Finds a blendable parameter of the given name. */
+	GAMEPLAYCAMERAS_API const FBlendableParameterInfo* FindBlendableParameter(FName ParameterName) const;
+	/** Finds a data parameter of the given name. */
+	GAMEPLAYCAMERAS_API const FDataParameterInfo* FindDataParameter(FName ParameterName) const;
+
+public:
+
+	/** Build this structure from the given camera node. */
+	GAMEPLAYCAMERAS_API void BuildFrom(UCameraNode* InCameraNode);
+
+private:
+
+	TArray<FBlendableParameterInfo> BlendableParameters;
+	TArray<FDataParameterInfo> DataParameters;
+
+	friend class FCameraNodeHierarchyBuilder;
+	friend class FCameraObjectInterfaceBuilder;
+	friend class FCameraObjectInterfaceParameterBuilder;
+	friend struct Internal::FInterfaceParameterBindingBuilder;
+};
+
+}  // namespace UE::Cameras
 
 /** Describes a custom camera blendable parameter. */
 USTRUCT()
@@ -31,7 +128,7 @@ struct FCustomCameraNodeBlendableParameter
 
 	/** The type of the parameter. */
 	UPROPERTY()
-	ECameraVariableType ParameterType = ECameraVariableType::Boolean;
+	ECameraVariableType VariableType = ECameraVariableType::Boolean;
 
 	/** The struct type of a blendable struct. */
 	UPROPERTY()
@@ -60,15 +157,15 @@ struct FCustomCameraNodeDataParameter
 
 	/** The type of the parameter. */
 	UPROPERTY()
-	ECameraContextDataType ParameterType = ECameraContextDataType::Name;
+	ECameraContextDataType DataType = ECameraContextDataType::Name;
 
 	/** The type of the parameter container. */
 	UPROPERTY()
-	ECameraContextDataContainerType ParameterContainerType = ECameraContextDataContainerType::None;
+	ECameraContextDataContainerType DataContainerType = ECameraContextDataContainerType::None;
 
 	/** An extra type object for the parameter. */
 	UPROPERTY()
-	TObjectPtr<const UObject> ParameterTypeObject;
+	TObjectPtr<const UObject> DataTypeObject;
 
 	/** An optional context data ID for dynamically driving the parameter's value. */
 	UPROPERTY()
@@ -77,7 +174,10 @@ struct FCustomCameraNodeDataParameter
 	bool operator==(const FCustomCameraNodeDataParameter& Other) const = default;
 };
 
-/** Describes custom camera parameters. */
+/** 
+ * Describes custom camera parameters. This structure is suitable for being owned by a camera node
+ * that wants to manage a variable number of custom parameters.
+ */
 USTRUCT()
 struct FCustomCameraNodeParameters
 {
@@ -100,84 +200,6 @@ struct FCustomCameraNodeParameters
 	bool operator==(const FCustomCameraNodeParameters& Other) const = default;
 };
 
-/**
- * A structure for providing custom camera rig parameter information.
- */
-struct FCustomCameraNodeParameterInfos
-{
-	/** Returns whether there are any blendable or data parameters. */
-	bool HasAnyParameters() const { return !BlendableParameters.IsEmpty() || !DataParameters.IsEmpty(); }
-
-	/** 
-	 * Declares a blendable parameter. 
-	 * Pass a null pointer for OverrideVariableID if this parameter should not be overridable
-	 * by a camera rig parameter.
-	 */
-	GAMEPLAYCAMERAS_API void AddBlendableParameter(
-			FName ParameterName, 
-			ECameraVariableType ParameterType, 
-			const UScriptStruct* BlendableStructType,
-			const uint8* DefaultValue,
-			FCameraVariableID* OverrideVariableID);
-
-	GAMEPLAYCAMERAS_API void AddBlendableParameter(FCustomCameraNodeBlendableParameter& Parameter, const uint8* DefaultValue);
-
-	/** 
-	 * Declares a data parameter.
-	 * Pass a null pointer for OverrideDataID if this parameter should not be overridable
-	 * by a camera rig parameter.
-	 */
-	GAMEPLAYCAMERAS_API void AddDataParameter(
-			FName ParameterName, 
-			ECameraContextDataType ParameterType,
-			ECameraContextDataContainerType ParameterContainerType,
-			const UObject* ParameterTypeObject,
-			const uint8* DefaultValue,
-			FCameraContextDataID* OverrideDataID);
-
-	GAMEPLAYCAMERAS_API void AddDataParameter(FCustomCameraNodeDataParameter& Parameter, const uint8* DefaultValue);
-
-	/** Gets the list of blendable parameters. */
-	GAMEPLAYCAMERAS_API void GetBlendableParameters(TArray<FCustomCameraNodeBlendableParameter>& OutBlendableParameters) const;
-	/** Gets the list of data parameters. */
-	GAMEPLAYCAMERAS_API void GetDataParameters(TArray<FCustomCameraNodeDataParameter>& OutDataParameters) const;
-
-	/** Finds a blendable parameter of the given name. */
-	GAMEPLAYCAMERAS_API bool FindBlendableParameter(FName ParameterName, FCustomCameraNodeBlendableParameter& OutParameter) const;
-	/** Finds a data parameter of the given name. */
-	GAMEPLAYCAMERAS_API bool FindDataParameter(FName ParameterName, FCustomCameraNodeDataParameter& OutParameter) const;
-
-private:
-
-	struct FBlendableParameterInfo
-	{
-		FName ParameterName;
-		ECameraVariableType ParameterType;
-		const UScriptStruct* BlendableStructType = nullptr;
-		const uint8* DefaultValue = nullptr;
-		FCameraVariableID* OverrideVariableID = nullptr;
-		UCameraVariableAsset* OverrideVariable = nullptr;
-	};
-
-	struct FDataParameterInfo
-	{
-		FName ParameterName;
-		ECameraContextDataType ParameterType;
-		ECameraContextDataContainerType ParameterContainerType;
-		const UObject* ParameterTypeObject;
-		const uint8* DefaultValue = nullptr;
-		FCameraContextDataID* OverrideDataID = nullptr;
-	};
-
-	TArray<FBlendableParameterInfo> BlendableParameters;
-	TArray<FDataParameterInfo> DataParameters;
-
-	friend class UE::Cameras::FCameraNodeHierarchyBuilder;
-	friend class UE::Cameras::FCameraObjectInterfaceBuilder;
-	friend class UE::Cameras::FCameraObjectInterfaceParameterBuilder;
-	friend struct UE::Cameras::Internal::FInterfaceParameterBindingBuilder;
-};
-
 UINTERFACE(MinimalAPI)
 class UCustomCameraNodeParameterProvider : public UInterface
 {
@@ -194,8 +216,10 @@ class ICustomCameraNodeParameterProvider
 
 public:
 
+	using FCameraNodeParameterInfos = UE::Cameras::FCameraNodeParameterInfos;
+
 	/** Gathers the custom parameters on this node. */
-	virtual void GetCustomCameraNodeParameters(FCustomCameraNodeParameterInfos& OutParameterInfos) {}
+	virtual void GetCustomCameraNodeParameters(FCameraNodeParameterInfos& OutParameterInfos) {}
 
 	/** Utility function for sub-classes to broadcast when the custom parameters have changed. */
 	GAMEPLAYCAMERAS_API void OnCustomCameraNodeParametersChanged(const UCameraNode* ThisAsCameraNode) const;
diff --git a/Engine/Plugins/Cameras/GameplayCameras/Source/GameplayCameras/Public/Nodes/Common/CameraRigCameraNode.h b/Engine/Plugins/Cameras/GameplayCameras/Source/GameplayCameras/Public/Nodes/Common/CameraRigCameraNode.h
--- a/Engine/Plugins/Cameras/GameplayCameras/Source/GameplayCameras/Public/Nodes/Common/CameraRigCameraNode.h
+++ b/Engine/Plugins/Cameras/GameplayCameras/Source/GameplayCameras/Public/Nodes/Common/CameraRigCameraNode.h
@@ -32,7 +32,7 @@ protected:
 	virtual void GatherPackages(FCameraRigPackages& OutPackages) const override;
 
 	// ICustomCameraNodeParameterProvider interface.
-	virtual void GetCustomCameraNodeParameters(FCustomCameraNodeParameterInfos& OutParameterInfos) override;
+	virtual void GetCustomCameraNodeParameters(FCameraNodeParameterInfos& OutParameterInfos) override;
 
 	// IObjectTreeGraphObject interface.
 #if WITH_EDITOR
@@ -74,11 +74,6 @@ protected:
 	virtual void OnBuild(const FCameraNodeEvaluatorBuildParams& Params) override;
 	virtual void OnRun(const FCameraNodeEvaluationParams& Params, FCameraNodeEvaluationResult& OutResult) override;
 
-private:
-
-	void ApplyParameterOverrides(FCameraVariableTable& OutVariableTable, bool bDrivenOnly);
-	void ApplyParameterOverrides(FCameraVariableTable& OutVariableTable, FCameraContextDataTable& OutContextDataTable, bool bDrivenOnly);
-
 private:
 
 	FCameraNodeEvaluator* CameraRigRootEvaluator = nullptr;
diff --git a/Engine/Plugins/Cameras/GameplayCameras/Source/GameplayCameras/Public/Nodes/Framing/BaseFramingCameraNode.h b/Engine/Plugins/Cameras/GameplayCameras/Source/GameplayCameras/Public/Nodes/Framing/BaseFramingCameraNode.h
--- a/Engine/Plugins/Cameras/GameplayCameras/Source/GameplayCameras/Public/Nodes/Framing/BaseFramingCameraNode.h
+++ b/Engine/Plugins/Cameras/GameplayCameras/Source/GameplayCameras/Public/Nodes/Framing/BaseFramingCameraNode.h
@@ -124,7 +124,7 @@ public:
 	virtual void PostLoad() override;
 
 	// ICustomCameraNodeParameterProvider interface.
-	virtual void GetCustomCameraNodeParameters(FCustomCameraNodeParameterInfos& OutParameterInfos) override;
+	virtual void GetCustomCameraNodeParameters(FCameraNodeParameterInfos& OutParameterInfos) override;
 };
 
 namespace UE::Cameras
diff --git a/Engine/Plugins/Cameras/GameplayCameras/Source/GameplayCameras/Public/Nodes/Shakes/CameraShakeCameraNode.h b/Engine/Plugins/Cameras/GameplayCameras/Source/GameplayCameras/Public/Nodes/Shakes/CameraShakeCameraNode.h
--- a/Engine/Plugins/Cameras/GameplayCameras/Source/GameplayCameras/Public/Nodes/Shakes/CameraShakeCameraNode.h
+++ b/Engine/Plugins/Cameras/GameplayCameras/Source/GameplayCameras/Public/Nodes/Shakes/CameraShakeCameraNode.h
@@ -34,7 +34,7 @@ protected:
 	virtual FCameraNodeEvaluatorPtr OnBuildEvaluator(FCameraNodeEvaluatorBuilder& Builder) const override;
 
 	// ICustomCameraNodeParameterProvider interface.
-	virtual void GetCustomCameraNodeParameters(FCustomCameraNodeParameterInfos& OutParameterInfos) override;
+	virtual void GetCustomCameraNodeParameters(FCameraNodeParameterInfos& OutParameterInfos) override;
 
 	// UObject interface.
 #if WITH_EDITOR
diff --git a/Engine/Plugins/Cameras/GameplayCameras/Source/GameplayCameras/Public/Nodes/Utility/BlueprintCameraNode.h b/Engine/Plugins/Cameras/GameplayCameras/Source/GameplayCameras/Public/Nodes/Utility/BlueprintCameraNode.h
--- a/Engine/Plugins/Cameras/GameplayCameras/Source/GameplayCameras/Public/Nodes/Utility/BlueprintCameraNode.h
+++ b/Engine/Plugins/Cameras/GameplayCameras/Source/GameplayCameras/Public/Nodes/Utility/BlueprintCameraNode.h
@@ -169,7 +169,7 @@ protected:
 #endif  // WITH_EDITOR
 
 	// ICustomCameraNodeParameterProvider interface.
-	virtual void GetCustomCameraNodeParameters(FCustomCameraNodeParameterInfos& OutParameterInfos) override;
+	virtual void GetCustomCameraNodeParameters(FCameraNodeParameterInfos& OutParameterInfos) override;
 
 	// UObject interface.
 	virtual void PostLoad() override;
diff --git a/Engine/Plugins/Cameras/GameplayCameras/Source/GameplayCameras/Public/StructUtils/OverridablePropertyBag.h b/Engine/Plugins/Cameras/GameplayCameras/Source/GameplayCameras/Public/StructUtils/OverridablePropertyBag.h
--- a/Engine/Plugins/Cameras/GameplayCameras/Source/GameplayCameras/Public/StructUtils/OverridablePropertyBag.h
+++ b/Engine/Plugins/Cameras/GameplayCameras/Source/GameplayCameras/Public/StructUtils/OverridablePropertyBag.h
@@ -40,6 +40,16 @@ public:
 	 */
 	GAMEPLAYCAMERAS_API void SetPropertyOverriden(const FGuid& InPropertyID, bool bIsOverriden);
 
+	/**
+	 * Returns whether the given property is overriden.
+	 */
+	GAMEPLAYCAMERAS_API bool IsPropertyOverriden(const FName InPropertyName) const;
+
+	/**
+	 * Sets whether the given property is overriden by setting the matching override property's value.
+	 */
+	GAMEPLAYCAMERAS_API void SetPropertyOverriden(const FName InPropertyName, bool bIsOverriden);
+
 public:
 
 	/**
diff --git a/Engine/Plugins/Cameras/GameplayCameras/Source/GameplayCamerasEditor/Private/Editors/CameraNodeGraphNode.cpp b/Engine/Plugins/Cameras/GameplayCameras/Source/GameplayCamerasEditor/Private/Editors/CameraNodeGraphNode.cpp
--- a/Engine/Plugins/Cameras/GameplayCameras/Source/GameplayCamerasEditor/Private/Editors/CameraNodeGraphNode.cpp
+++ b/Engine/Plugins/Cameras/GameplayCameras/Source/GameplayCamerasEditor/Private/Editors/CameraNodeGraphNode.cpp
@@ -51,6 +51,8 @@ void UCameraNodeGraphNode::OnCustomCameraNodeParametersChanged(const UCameraNode
 
 void UCameraNodeGraphNode::AllocateDefaultPins()
 {
+	using namespace UE::Cameras;
+
 	Super::AllocateDefaultPins();
 
 	// Add extra input pins for any camera parameter, variable reference, and context data.
@@ -154,21 +156,18 @@ void UCameraNodeGraphNode::AllocateDefaultPins()
 	ICustomCameraNodeParameterProvider* CustomParameterProvider = Cast<ICustomCameraNodeParameterProvider>(Object);
 	if (CustomParameterProvider)
 	{
-		FCustomCameraNodeParameterInfos CustomParameters;
-		CustomParameterProvider->GetCustomCameraNodeParameters(CustomParameters);
+		FCameraNodeParameterInfos ObjectParameters;
+		CustomParameterProvider->GetCustomCameraNodeParameters(ObjectParameters);
 
 		// Add pins for blendable parameters.
-		TArray<FCustomCameraNodeBlendableParameter> CustomBlendableParameters;
-		CustomParameters.GetBlendableParameters(CustomBlendableParameters);
-
 		const UEnum* VariableTypeEnum = StaticEnum<ECameraVariableType>();
-		for (const FCustomCameraNodeBlendableParameter& BlendableParameter : CustomBlendableParameters)
+		for (const FCameraNodeBlendableParameterInfo& BlendableParameter : ObjectParameters.GetBlendableParameters())
 		{
 			FEdGraphPinType PinType;
 			PinType.PinCategory = UCameraObjectGraphSchemaBase::PC_CameraParameter;
-			PinType.PinSubCategory = VariableTypeEnum->GetValueAsName(BlendableParameter.ParameterType);
+			PinType.PinSubCategory = VariableTypeEnum->GetValueAsName(BlendableParameter.VariableType);
 
-			switch (BlendableParameter.ParameterType)
+			switch (BlendableParameter.VariableType)
 			{
 #define UE_CAMERA_VARIABLE_FOR_TYPE(ValueType, ValueName)\
 				case ECameraVariableType::ValueName:\
@@ -177,28 +176,25 @@ void UCameraNodeGraphNode::AllocateDefaultPins()
 				UE_CAMERA_VARIABLE_FOR_ALL_TYPES()
 #undef UE_CAMERA_VARIABLE_FOR_TYPE
 				case ECameraVariableType::BlendableStruct:
-					PinType.PinSubCategoryObject = const_cast<UScriptStruct*>(BlendableParameter.BlendableStructType.Get());
+					PinType.PinSubCategoryObject = const_cast<UScriptStruct*>(BlendableParameter.BlendableStructType);
 					break;
 			}
 
 			UEdGraphPin* ParameterPin = CreatePin(EGPD_Input, PinType, BlendableParameter.ParameterName);
 			ParameterPin->PinFriendlyName = FText::FromName(BlendableParameter.ParameterName);
-			ParameterPin->PinToolTip = VariableTypeEnum->GetNameStringByValue((int64)BlendableParameter.ParameterType);
+			ParameterPin->PinToolTip = VariableTypeEnum->GetNameStringByValue((int64)BlendableParameter.VariableType);
 		}
 
 		// Add pins for data parameters.
-		TArray<FCustomCameraNodeDataParameter> CustomDataParameters;
-		CustomParameters.GetDataParameters(CustomDataParameters);
 		const UEnum* DataTypeEnum = StaticEnum<ECameraContextDataType>();
-
-		for (const FCustomCameraNodeDataParameter& DataParameter : CustomDataParameters)
+		for (const FCameraNodeDataParameterInfo& DataParameter : ObjectParameters.GetDataParameters())
 		{
 			FEdGraphPinType PinType;
 			PinType.PinCategory = UCameraObjectGraphSchemaBase::PC_CameraContextData;
-			PinType.PinSubCategory = DataTypeEnum->GetNameByValue((int64)DataParameter.ParameterType);
-			PinType.PinSubCategoryObject = const_cast<UObject*>(DataParameter.ParameterTypeObject.Get());
+			PinType.PinSubCategory = DataTypeEnum->GetNameByValue((int64)DataParameter.DataType);
+			PinType.PinSubCategoryObject = const_cast<UObject*>(DataParameter.DataTypeObject);
 
-			if (DataParameter.ParameterContainerType == ECameraContextDataContainerType::Array)
+			if (DataParameter.DataContainerType == ECameraContextDataContainerType::Array)
 			{
 				PinType.ContainerType = EPinContainerType::Array;
 			}
@@ -206,10 +202,10 @@ void UCameraNodeGraphNode::AllocateDefaultPins()
 			UEdGraphPin* ContextDataPin = CreatePin(EGPD_Input, PinType, DataParameter.ParameterName);
 			ContextDataPin->PinFriendlyName = FText::FromName(DataParameter.ParameterName);
 
-			FString PinToolTip = DataTypeEnum->GetNameStringByValue((int64)DataParameter.ParameterType);
-			if (DataParameter.ParameterTypeObject)
+			FString PinToolTip = DataTypeEnum->GetNameStringByValue((int64)DataParameter.DataType);
+			if (DataParameter.DataTypeObject)
 			{
-				PinToolTip = DataParameter.ParameterTypeObject->GetName();
+				PinToolTip = DataParameter.DataTypeObject->GetName();
 			}
 			ContextDataPin->PinToolTip = PinToolTip;
 		}
diff --git a/Engine/Plugins/Cameras/GameplayCameras/Source/GameplayCamerasEditor/Private/Editors/CameraObjectGraphSchemaBase.cpp b/Engine/Plugins/Cameras/GameplayCameras/Source/GameplayCamerasEditor/Private/Editors/CameraObjectGraphSchemaBase.cpp
--- a/Engine/Plugins/Cameras/GameplayCameras/Source/GameplayCamerasEditor/Private/Editors/CameraObjectGraphSchemaBase.cpp
+++ b/Engine/Plugins/Cameras/GameplayCameras/Source/GameplayCamerasEditor/Private/Editors/CameraObjectGraphSchemaBase.cpp
@@ -123,6 +123,8 @@ UEdGraphPin* UCameraObjectGraphSchemaBase::FindPinByType(UEdGraphNode* InNode, c
 
 void UCameraObjectGraphSchemaBase::GetGraphContextActions(FGraphContextMenuBuilder& ContextMenuBuilder) const
 {
+	using namespace UE::Cameras;
+
 	// See if we were dragging a camera parameter pin or camera variable reference pin.
 	if (const UEdGraphPin* DraggedPin = ContextMenuBuilder.FromPin)
 	{
@@ -144,11 +146,11 @@ void UCameraObjectGraphSchemaBase::GetGraphContextActions(FGraphContextMenuBuild
 			const UClass* CameraNodeClass = CameraNodeNode->GetObject()->GetClass();
 			FProperty* Property = CameraNodeClass->FindPropertyByName(DraggedPin->PinName);
 
-			FCustomCameraNodeParameterInfos CustomParameters;
+			FCameraNodeParameterInfos CameraNodeParameters;
 			ICustomCameraNodeParameterProvider* CustomParameterProvider = Cast<ICustomCameraNodeParameterProvider>(CameraNodeNode->GetObject());
 			if (CustomParameterProvider)
 			{
-				CustomParameterProvider->GetCustomCameraNodeParameters(CustomParameters);
+				CustomParameterProvider->GetCustomCameraNodeParameters(CameraNodeParameters);
 			}
 
 			TSharedRef<FCameraObjectGraphSchemaAction_NewInterfaceParameterNode> Action = 
@@ -163,11 +165,10 @@ void UCameraObjectGraphSchemaBase::GetGraphContextActions(FGraphContextMenuBuild
 				ECameraVariableType VariableType;
 				const UScriptStruct* BlendableStructType = nullptr;
 
-				FCustomCameraNodeBlendableParameter BlendableParameter;
-				if (CustomParameters.FindBlendableParameter(DraggedPin->PinName, BlendableParameter))
+				if (const FCameraNodeBlendableParameterInfo* BlendableParameter = CameraNodeParameters.FindBlendableParameter(DraggedPin->PinName))
 				{
-					VariableType = BlendableParameter.ParameterType;
-					BlendableStructType = BlendableParameter.BlendableStructType;
+					VariableType = BlendableParameter->VariableType;
+					BlendableStructType = BlendableParameter->BlendableStructType;
 				}
 				else if (FStructProperty* StructProperty = CastField<FStructProperty>(Property))
 				{
@@ -206,12 +207,11 @@ void UCameraObjectGraphSchemaBase::GetGraphContextActions(FGraphContextMenuBuild
 				ECameraContextDataContainerType DataContainerType = ECameraContextDataContainerType::None;
 				const UObject* DataTypeObject = nullptr;
 				
-				FCustomCameraNodeDataParameter DataParameter;
-				if (CustomParameters.FindDataParameter(DraggedPin->PinName, DataParameter))
+				if (const FCameraNodeDataParameterInfo* DataParameter = CameraNodeParameters.FindDataParameter(DraggedPin->PinName))
 				{
-					DataType = DataParameter.ParameterType;
-					DataContainerType = DataParameter.ParameterContainerType;
-					DataTypeObject = DataParameter.ParameterTypeObject;
+					DataType = DataParameter->DataType;
+					DataContainerType = DataParameter->DataContainerType;
+					DataTypeObject = DataParameter->DataTypeObject;
 				}
 				else if (Property)
 				{
@@ -296,7 +296,7 @@ const FPinConnectionResponse UCameraObjectGraphSchemaBase::CanCreateConnection(c
 			UCameraObjectInterfaceParameterBase* Parameter = ParameterGetter->GetInterfaceParameter();
 			UCameraObjectInterfaceBlendableParameter* BlendableParameter = Cast<UCameraObjectInterfaceBlendableParameter>(Parameter);
 			if (BlendableParameter && 
-					A->PinType.PinSubCategory == UEnum::GetValueAsName(BlendableParameter->ParameterType))
+					A->PinType.PinSubCategory == UEnum::GetValueAsName(BlendableParameter->VariableType))
 			{
 				return FPinConnectionResponse(CONNECT_RESPONSE_MAKE, TEXT("Compatible pin types"));
 			}
@@ -322,7 +322,7 @@ const FPinConnectionResponse UCameraObjectGraphSchemaBase::CanCreateConnection(c
 			UCameraObjectInterfaceParameterBase* Parameter = ParameterGetter->GetInterfaceParameter();
 			UCameraObjectInterfaceBlendableParameter* BlendableParameter = Cast<UCameraObjectInterfaceBlendableParameter>(Parameter);
 			if (BlendableParameter && 
-					B->PinType.PinSubCategory == UEnum::GetValueAsName(BlendableParameter->ParameterType))
+					B->PinType.PinSubCategory == UEnum::GetValueAsName(BlendableParameter->VariableType))
 			{
 				return FPinConnectionResponse(CONNECT_RESPONSE_MAKE, TEXT("Compatible pin types"));
 			}
@@ -572,7 +572,7 @@ UEdGraphNode* FCameraObjectGraphSchemaAction_NewInterfaceParameterNode::PerformA
 	if (ParameterDefinition.ParameterType == ECameraObjectInterfaceParameterType::Blendable)
 	{
 		UCameraObjectInterfaceBlendableParameter* NewBlendableParameter = NewObject<UCameraObjectInterfaceBlendableParameter>(CameraObject, NAME_None, RF_Transactional);
-		NewBlendableParameter->ParameterType = ParameterDefinition.VariableType;
+		NewBlendableParameter->VariableType = ParameterDefinition.VariableType;
 		NewBlendableParameter->BlendableStructType = ParameterDefinition.BlendableStructType;
 		CameraObject->Interface.BlendableParameters.Add(NewBlendableParameter);
 		NewInterfaceParameter = NewBlendableParameter;
diff --git a/Engine/Plugins/Cameras/GameplayCameras/Source/GameplayCamerasEditor/Private/Editors/SCameraObjectInterfaceParametersPanel.cpp b/Engine/Plugins/Cameras/GameplayCameras/Source/GameplayCamerasEditor/Private/Editors/SCameraObjectInterfaceParametersPanel.cpp
--- a/Engine/Plugins/Cameras/GameplayCameras/Source/GameplayCamerasEditor/Private/Editors/SCameraObjectInterfaceParametersPanel.cpp
+++ b/Engine/Plugins/Cameras/GameplayCameras/Source/GameplayCamerasEditor/Private/Editors/SCameraObjectInterfaceParametersPanel.cpp
@@ -197,7 +197,7 @@ protected:
 	FEdGraphPinType GetBlendableParameterPinType() const
 	{
 		FEdGraphPinType PinType;
-		switch (Item->ParameterType)
+		switch (Item->VariableType)
 		{
 			case ECameraVariableType::Boolean:
 				PinType.PinCategory = UEdGraphSchema_K2::PC_Boolean;
@@ -333,12 +333,12 @@ protected:
 			bIsValidType = false;
 		}
 
-		if (ensure(bIsValidType) && Item->ParameterType != NewParameterType)
+		if (ensure(bIsValidType) && Item->VariableType != NewParameterType)
 		{
 			const FScopedTransaction Transaction(LOCTEXT("ChangeBlendableParameterType", "Change Blendable Parameter Type"));
 
 			Item->Modify();
-			Item->ParameterType = NewParameterType;
+			Item->VariableType = NewParameterType;
 		}
 	}
 
diff --git a/Engine/Plugins/Cameras/GameplayCameras/Source/GameplayCamerasUncookedOnly/Private/BlueprintGraph/K2Node_CameraRigBase.cpp b/Engine/Plugins/Cameras/GameplayCameras/Source/GameplayCamerasUncookedOnly/Private/BlueprintGraph/K2Node_CameraRigBase.cpp
--- a/Engine/Plugins/Cameras/GameplayCameras/Source/GameplayCamerasUncookedOnly/Private/BlueprintGraph/K2Node_CameraRigBase.cpp
+++ b/Engine/Plugins/Cameras/GameplayCameras/Source/GameplayCamerasUncookedOnly/Private/BlueprintGraph/K2Node_CameraRigBase.cpp
@@ -140,7 +140,7 @@ void UK2Node_CameraRigBase::OnCameraRigAssetBuilt(const UCameraRigAsset* InBuilt
 
 FEdGraphPinType UK2Node_CameraRigBase::MakeBlendableParameterPinType(const UCameraObjectInterfaceBlendableParameter* BlendableParameter)
 {
-	return MakeBlendableParameterPinType(BlendableParameter->ParameterType, BlendableParameter->BlendableStructType);
+	return MakeBlendableParameterPinType(BlendableParameter->VariableType, BlendableParameter->BlendableStructType);
 }
 
 FEdGraphPinType UK2Node_CameraRigBase::MakeBlendableParameterPinType(ECameraVariableType CameraVariableType, const UScriptStruct* BlendableStructType)
diff --git a/Engine/Plugins/Cameras/GameplayCameras/Source/GameplayCamerasUncookedOnly/Private/BlueprintGraph/K2Node_SingleCameraRigParameterBase.cpp b/Engine/Plugins/Cameras/GameplayCameras/Source/GameplayCamerasUncookedOnly/Private/BlueprintGraph/K2Node_SingleCameraRigParameterBase.cpp
--- a/Engine/Plugins/Cameras/GameplayCameras/Source/GameplayCamerasUncookedOnly/Private/BlueprintGraph/K2Node_SingleCameraRigParameterBase.cpp
+++ b/Engine/Plugins/Cameras/GameplayCameras/Source/GameplayCamerasUncookedOnly/Private/BlueprintGraph/K2Node_SingleCameraRigParameterBase.cpp
@@ -33,7 +33,7 @@ void UK2Node_SingleCameraRigParameterBase::Initialize(const FAssetData& Unloaded
 		{
 			if (const UCameraObjectInterfaceBlendableParameter* BlendableParameter = LoadedCameraRig->Interface.FindBlendableParameterByName(InCameraParameterName))
 			{
-				Initialize(LoadedCameraRig, InCameraParameterName, BlendableParameter->ParameterType, BlendableParameter->BlendableStructType);
+				Initialize(LoadedCameraRig, InCameraParameterName, BlendableParameter->VariableType, BlendableParameter->BlendableStructType);
 			}
 			else if (const UCameraObjectInterfaceDataParameter* DataParameter = LoadedCameraRig->Interface.FindDataParameterByName(InCameraParameterName))
 			{
