Cameras: fix asset dirtying problems between camera rigs and camera assets

The first problem is that camera rigs build all the allocation information they need, including the memory needed for their evaluator tree, and the entries they need in the variables/data tables. This can easily mark a camera rig asset as modified, when (1) evaluator classes change (e.g. a new field is added) or (2) when a referenced camera rig, such as a prefab rig, changes in a way that affects its own allocation info. This ends up prompting the user to re-save assets that they haven't touched, with no obvious reason why.

So the solution to this is easy: only build allocation info during cooking. In non-cooked builds, we just allocate things as we go, which is less performant, but much better for the user.

The second problem is that camera assets were automatically gathering all the exposed parameters from all their referenced camera rigs, and using that as their own exposed parameters. This was simple and functional at first, but it also created the same problem: the camera asset would get flagged as modified when any referenced camera rig had new/changed/deleted parameters. This again prompted the user to re-save the camera asset for no good reason. In addition to that, as camera assets grew in scope, they could start getting a lot of parameters that the user didn't necessarily want.

Solving this problem requires a new workflow: the camera asset editor now has an "Interface Parameters" panel, similar to the camera rig editor. It lets the user add/remove parameters on the camera asset itself, and specify how each of these parameters binds to which camera rig's parameter. It effectively removes the automatic step of gathering all referenced camera rigs' parameters, and replaces it by a manual step where the user can pick and choose which parameters to expose on the camera asset. This way, modifying a referenced camera rig doesn't modify the camera asset: the user has to explicitly modify the camera asset if they want its interface to change.

Other notes about this change:

- All the builder classes and methods now pass around a base "build context", which contains information about whether we're cooking or not.
- Camera assets are upgraded on build or open to get a default list of interface parameters that reproduces what was automatically built before.
- When initializing a new camera rig instance, we now ensure that the variable/data tables have all the entries needed. This makes uncooked builds work by lazily allocated memory when we don't have the allocation info.
- The camera director's GatherUsageRigInfo now does extra stuff out of the box, such as looking into the proxy redirect table, and removing unwanted things from the output.
- The code that converts a variable/data entry info a graph editor pin type has been moved into a common place for reuse.

#Synced-CL 49781371

diff --git a/Engine/Plugins/Cameras/GameplayCameras/Source/GameplayCameras/Private/Build/CameraAssetBuilder.cpp b/Engine/Plugins/Cameras/GameplayCameras/Source/GameplayCameras/Private/Build/CameraAssetBuilder.cpp
--- a/Engine/Plugins/Cameras/GameplayCameras/Source/GameplayCameras/Private/Build/CameraAssetBuilder.cpp
+++ b/Engine/Plugins/Cameras/GameplayCameras/Source/GameplayCameras/Private/Build/CameraAssetBuilder.cpp
@@ -2,14 +2,13 @@
 
 #include "Build/CameraAssetBuilder.h"
 
+#include "Build/CameraBuildContext.h"
 #include "Build/CameraObjectInterfaceParameterBuilder.h"
 #include "Build/CameraRigAssetBuilder.h"
 #include "Core/CameraAsset.h"
 #include "Core/CameraDirector.h"
 #include "Core/CameraObjectInterfaceParameterDefinition.h"
-#include "Core/CameraParameters.h"
 #include "Core/CameraRigAsset.h"
-#include "Core/CameraRigProxyAsset.h"
 #include "GameplayCamerasDelegates.h"
 #include "Logging/TokenizedMessage.h"
 
@@ -18,8 +17,80 @@
 namespace UE::Cameras
 {
 
-FCameraAssetBuilder::FCameraAssetBuilder(FCameraBuildLog& InBuildLog)
-	: BuildLog(InBuildLog)
+namespace Private
+{
+
+void GatherCameraAssetRigUsage(const UCameraAsset* InCameraAsset, FCameraDirectorRigUsageInfo& OutUsageInfo)
+{
+	UCameraDirector* CameraDirector = InCameraAsset->GetCameraDirector();
+	if (CameraDirector)
+	{
+		CameraDirector->GatherRigUsageInfo(OutUsageInfo);
+	}
+}
+
+void BuildCameraAssetDefaultParameterDefinitions(UCameraAsset* CameraAsset, TArrayView<UCameraRigAsset*> AllCameraRigs)
+{
+	// Get the list of all the camera rigs' interface parameters.
+	TMap<FName, TArray<UCameraRigAsset*>> UsedParameterNames;
+	TMap<UCameraRigAsset*, TArray<FCameraObjectInterfaceParameterDefinition>> DefinitionsByCameraRig;
+	for (UCameraRigAsset* CameraRig : AllCameraRigs)
+	{
+		if (ensure(!DefinitionsByCameraRig.Contains(CameraRig)))
+		{
+			TArray<FCameraObjectInterfaceParameterDefinition>& DefinitionsForCameraRig = DefinitionsByCameraRig.Add(CameraRig);
+			for (const FCameraObjectInterfaceParameterDefinition& Definition : CameraRig->GetParameterDefinitions())
+			{
+				if (Definition.bIsVisible)
+				{
+					DefinitionsForCameraRig.Add(Definition);
+					UsedParameterNames.FindOrAdd(Definition.ParameterName).Add(CameraRig);
+				}
+			}
+		}
+	}
+	// Resolve name conflicts.
+	TMap<TTuple<UCameraRigAsset*, FName>, FString> SourceParameterToInterfaceName;
+	for (const TPair<FName, TArray<UCameraRigAsset*>>& Pair : UsedParameterNames)
+	{
+		const TArray<UCameraRigAsset*>& ConflictCameraRigs(Pair.Value);
+		if (ConflictCameraRigs.Num() > 1)
+		{
+			// For each parameter with the same name, append the source camera rig's name.
+			const FName ConflictName(Pair.Key);
+			for (UCameraRigAsset* CameraRig : ConflictCameraRigs)
+			{
+				const FString NewName = FString::Format(TEXT("{0}_{1}"), { *GetNameSafe(CameraRig), *ConflictName.ToString() });
+				SourceParameterToInterfaceName.Add({ CameraRig, Pair.Key }, NewName);
+			}
+		}
+		// else: the interface name will be the same as the source parameter's.
+	}
+
+	// Build the final list of parameter definitions.
+	for (const TPair<UCameraRigAsset*, TArray<FCameraObjectInterfaceParameterDefinition>>& Pair : DefinitionsByCameraRig)
+	{
+		UCameraRigAsset* CameraRig = Pair.Key;
+		const TArray<FCameraObjectInterfaceParameterDefinition>& DefinitionsForCameraRig(Pair.Value);
+		for (const FCameraObjectInterfaceParameterDefinition& Definition : DefinitionsForCameraRig)
+		{
+			const FString DefaultInterfaceName = Definition.ParameterName.ToString();
+			const FString ResolvedInterfaceName = SourceParameterToInterfaceName.FindRef(
+					{ CameraRig, Definition.ParameterName }, DefaultInterfaceName);
+
+			UCameraAssetInterfaceParameter* NewInterfaceParameter = NewObject<UCameraAssetInterfaceParameter>(CameraAsset);
+			NewInterfaceParameter->InterfaceParameterName = ResolvedInterfaceName;
+			NewInterfaceParameter->SourceCameraRig = CameraRig;
+			NewInterfaceParameter->SourceParameterName = Definition.ParameterName;
+			CameraAsset->Interface.Parameters.Add(NewInterfaceParameter);
+		}
+	}
+}
+
+}  // namespace Private
+
+FCameraAssetBuilder::FCameraAssetBuilder(FCameraBuildContext& InBuildContext)
+	: BuildContext(InBuildContext)
 {
 }
 
@@ -31,11 +102,11 @@ void FCameraAssetBuilder::BuildCamera(UCameraAsset* InCameraAsset, bool bBuildRe
 	}
 
 	CameraAsset = InCameraAsset;
-	BuildLog.SetLoggingPrefix(InCameraAsset->GetPathName() + TEXT(": "));
+	BuildContext.BuildLog.SetLoggingPrefix(InCameraAsset->GetPathName() + TEXT(": "));
 	{
 		BuildCameraImpl(bBuildReferencedAssets);
 	}
-	BuildLog.SetLoggingPrefix(FString());
+	BuildContext.BuildLog.SetLoggingPrefix(FString());
 	UpdateBuildStatus();
 
 	FGameplayCamerasDelegates::OnCameraAssetBuilt().Broadcast(CameraAsset);
@@ -43,153 +114,185 @@ void FCameraAssetBuilder::BuildCamera(UCameraAsset* InCameraAsset, bool bBuildRe
 
 void FCameraAssetBuilder::BuildCameraImpl(bool bBuildReferencedAssets)
 {
-	TSet<UCameraRigAsset*> AllCameraRigs;
-	TSet<UCameraRigProxyAsset*> AllCameraRigProxies;
-
-	// Build the camera director and get the list of camera rigs it references.
+	// Build the camera director.
 	UCameraDirector* CameraDirector = CameraAsset->GetCameraDirector();
 	if (CameraDirector)
 	{
-		CameraDirector->BuildCameraDirector(BuildLog);
-
-		FCameraDirectorRigUsageInfo UsageInfo;
-		CameraDirector->GatherRigUsageInfo(UsageInfo);
-
-		AllCameraRigs.Append(UsageInfo.CameraRigs);
-		AllCameraRigProxies.Append(UsageInfo.CameraRigProxies);
+		CameraDirector->BuildCameraDirector(BuildContext);
 	}
 	else
 	{
-		BuildLog.AddMessage(EMessageSeverity::Error, LOCTEXT("MissingDirector", "Camera has no director set."));
+		BuildContext.BuildLog.AddMessage(EMessageSeverity::Error, LOCTEXT("MissingDirector", "Camera has no director set."));
 	}
 
-	// Some directors might add CDOs for some reason (for instance if they use a reference collector).
-	// We need to remove them and leave them untouched.
-	AllCameraRigs.Remove(GetDefault<UCameraRigAsset>());
-	AllCameraRigProxies.Remove(GetDefault<UCameraRigProxyAsset>());
-
-	if (AllCameraRigs.IsEmpty() && AllCameraRigProxies.IsEmpty())
+	// Get the list of referenced camera rigs and proxies.
+	FCameraDirectorRigUsageInfo UsageInfo;
+	Private::GatherCameraAssetRigUsage(CameraAsset, UsageInfo);
+	if (CameraDirector && UsageInfo.CameraRigs.IsEmpty() && UsageInfo.CameraRigProxies.IsEmpty())
 	{
-		BuildLog.AddMessage(EMessageSeverity::Warning, LOCTEXT("MissingRigs", "Camera director isn't using any camera rigs or proxies."));
+		BuildContext.BuildLog.AddMessage(EMessageSeverity::Warning, LOCTEXT("MissingRigs", "Camera director isn't using any camera rigs or proxies."));
 	}
 
-	// Get also the rigs from the proxy table.
-	if (CameraDirector)
+	BuiltReferencedAssets.Reset();
+	if (bBuildReferencedAssets)
 	{
-		for (const FCameraRigProxyRedirectTableEntry& Entry : CameraDirector->CameraRigProxyRedirectTable.Entries)
+		// Build each of the camera rigs.
+		for (UCameraRigAsset* CameraRig : UsageInfo.CameraRigs)
 		{
-			if (Entry.CameraRig)
+			const bool bNeedsBuild = CameraRig->BuildStatus == ECameraBuildStatus::Dirty;
+			if (bNeedsBuild)
 			{
-				AllCameraRigs.Add(Entry.CameraRig);
+				FCameraRigAssetBuilder CameraRigBuilder(BuildContext);
+				CameraRigBuilder.BuildCameraRig(CameraRig);
+				BuiltReferencedAssets.Add(CameraRig);
 			}
 		}
 	}
 
-	if (bBuildReferencedAssets)
+	// If we have some old data, we may need to build a default interface for it, but don't dirty the asset unless we're cooking and
+	// want to save the proper data. In other contexts, we don't want to cause the user to be prompted with an asset to save without
+	// them knowing what changed and why.
+	if (CameraAsset->bNeedsDefaultInterface)
 	{
-		// Build each of the camera rigs.
-		for (UCameraRigAsset* CameraRig : AllCameraRigs)
+		if (BuildContext.IsCooking())
 		{
-			FCameraRigAssetBuilder CameraRigBuilder(BuildLog);
-			CameraRigBuilder.BuildCameraRig(CameraRig);
+			CameraAsset->Modify();
 		}
+		Private::BuildCameraAssetDefaultParameterDefinitions(CameraAsset, UsageInfo.CameraRigs);
+		CameraAsset->bNeedsDefaultInterface = false;
 	}
 
-	// Get the list of all the camera rigs' interface parameters.
-	TMap<FName, TArray<const UCameraRigAsset*>> UsedParameterNames;
-	TMap<const UCameraRigAsset*, TArray<FCameraObjectInterfaceParameterDefinition>> DefinitionsByCameraRig;
-	for (const UCameraRigAsset* CameraRig : AllCameraRigs)
+	// Build the new parameter definitions, and then the new default parameters struct.
+	BuildParameters();
+
+	// Accumulate all the camera rigs' allocation infos and store that on the asset.
+	// Only do this when cooking. In editor builds we don't want to mark the asset as modified just because someone
+	// changed some parameters on a camera rig, or whatever.
+	if (BuildContext.IsCooking())
 	{
-		if (ensure(!DefinitionsByCameraRig.Contains(CameraRig)))
+		FCameraAssetAllocationInfo AllocationInfo;
+		for (const UCameraRigAsset* CameraRig : UsageInfo.CameraRigs)
 		{
-			TArray<FCameraObjectInterfaceParameterDefinition>& DefinitionsForCameraRig = DefinitionsByCameraRig.Add(CameraRig);
-			DefinitionsForCameraRig.Append(CameraRig->GetParameterDefinitions());
-
-			for (const FCameraObjectInterfaceParameterDefinition& Definition : DefinitionsForCameraRig)
-			{
-				UsedParameterNames.FindOrAdd(Definition.ParameterName).Add(CameraRig);
-			}
+			AllocationInfo.VariableTableInfo.Combine(CameraRig->AllocationInfo.VariableTableInfo);
+			AllocationInfo.ContextDataTableInfo.Combine(CameraRig->AllocationInfo.ContextDataTableInfo);
 		}
-	}
-	// Resolve name conflicts.
-	// We can safely change the parameter definitions and property bag property names because the look-ups
-	// for setting the default values afterwards are using Guids.
-	for (const TPair<FName, TArray<const UCameraRigAsset*>>& Pair : UsedParameterNames)
-	{
-		const TArray<const UCameraRigAsset*>& ConflictCameraRigs(Pair.Value);
-		if (ConflictCameraRigs.Num() > 1)
+
+		if (AllocationInfo != CameraAsset->AllocationInfo)
 		{
-			const FName ConflictName(Pair.Key);
-			for (const UCameraRigAsset* CameraRig : ConflictCameraRigs)
-			{
-				TArray<FCameraObjectInterfaceParameterDefinition>& Definitions = DefinitionsByCameraRig.FindChecked(CameraRig);
-				FCameraObjectInterfaceParameterDefinition* ConflictDefinition = Definitions.FindByPredicate(
-						[ConflictName](FCameraObjectInterfaceParameterDefinition& Item)
-						{
-							return Item.ParameterName == ConflictName;
-						});
-				if (ensure(ConflictDefinition))
-				{
-					const FString NewName = FString::Format(TEXT("{0}_{1}"), { *GetNameSafe(CameraRig), *ConflictName.ToString() });
-					ConflictDefinition->ParameterName = FName(NewName);
-				}
-			}
+			CameraAsset->Modify();
+			CameraAsset->AllocationInfo = AllocationInfo;
 		}
 	}
+}
 
-	// Build the final list of parameter definitions.
+void FCameraAssetBuilder::BuildParameters()
+{
+	// First build the new parameter definitions. Keep track of which ones are exposed from which camera rig.
 	TArray<FCameraObjectInterfaceParameterDefinition> ParameterDefinitions;
-	for (const TPair<const UCameraRigAsset*, TArray<FCameraObjectInterfaceParameterDefinition>>& Pair : DefinitionsByCameraRig)
+	TArray<const UCameraRigAsset*> SourceCameraRigs;
+
+	for (auto It = CameraAsset->Interface.Parameters.CreateIterator(); It; ++It)
 	{
-		const TArray<FCameraObjectInterfaceParameterDefinition>& DefinitionsForCameraRig(Pair.Value);
-		for (const FCameraObjectInterfaceParameterDefinition& Definition : DefinitionsForCameraRig)
+		const UCameraAssetInterfaceParameter* InterfaceParameter(*It);
+		if (!InterfaceParameter)
+		{
+			BuildContext.BuildLog.AddMessage(EMessageSeverity::Warning,
+					CameraAsset,
+					LOCTEXT("InvalidInterfaceParameter", "Invalid interface parameter was found and removed."));
+
+			CameraAsset->Modify();
+			It.RemoveCurrent();
+			continue;
+		}
+
+		if (InterfaceParameter->InterfaceParameterName.IsEmpty())
+		{
+			BuildContext.BuildLog.AddMessage(EMessageSeverity::Error,
+					InterfaceParameter,
+					LOCTEXT(
+						"InvalidInterfaceParameterName",
+						"Invalid interface parameter name."));
+			continue;
+		}
+
+		const UCameraRigAsset* CameraRig = InterfaceParameter->SourceCameraRig.Get();
+		if (!CameraRig)
+		{
+			BuildContext.BuildLog.AddMessage(EMessageSeverity::Error,
+					InterfaceParameter,
+					FText::Format(LOCTEXT(
+						"InterfaceParameterMissingSourceCameraRig",
+						"Interface parameter references a moved or deleted camera rig asset: {0}"),
+						FText::FromString(InterfaceParameter->SourceCameraRig.GetAssetName())));
+			continue;
+		}
+
+		if (InterfaceParameter->SourceParameterName.IsNone())
+		{
+			BuildContext.BuildLog.AddMessage(EMessageSeverity::Error,
+					InterfaceParameter,
+					FText::Format(LOCTEXT(
+						"InterfaceParameterNoSourceParameter",
+						"Interface parameter doesn't specify the parameter name to reference on camera rig '{0}'."),
+						FText::FromName(CameraRig->GetFName())));
+			continue;
+		}
+
+		FCameraObjectInterfaceParameterDefinition SourceParameterDefinition;
+		const bool bFoundSourceParameter = CameraRig->FindParameterDefinitionByName(
+				InterfaceParameter->SourceParameterName, SourceParameterDefinition);
+		if (!bFoundSourceParameter)
 		{
-			ParameterDefinitions.Add(Definition);
+			BuildContext.BuildLog.AddMessage(EMessageSeverity::Error,
+					InterfaceParameter,
+					FText::Format(LOCTEXT(
+						"InterfaceParameterMissingSourceParameter",
+						"Interface parameter references renamed or deleted parameter '{0}' on camera rig '{1}'."),
+						FText::FromName(InterfaceParameter->SourceParameterName),
+						FText::FromName(CameraRig->GetFName())));
+			continue;
 		}
+
+		FCameraObjectInterfaceParameterDefinition ParameterDefinition(SourceParameterDefinition);
+		ParameterDefinition.ParameterName = FName(InterfaceParameter->InterfaceParameterName);
+		ParameterDefinitions.Add(ParameterDefinition);
+
+		SourceCameraRigs.Add(CameraRig);
 	}
 
-	if (ParameterDefinitions != CameraAsset->ParameterDefinitions)
+	if (CameraAsset->ParameterDefinitions != ParameterDefinitions)
 	{
 		CameraAsset->Modify();
 		CameraAsset->ParameterDefinitions = ParameterDefinitions;
 	}
 
-	// Rebuild the default parameters property bag.
+	// Now build the new default parameters structure.
 	TArray<FPropertyBagPropertyDesc> DefaultParameterProperties;
 	FCameraObjectInterfaceParameterBuilder::AppendDefaultParameterProperties(
 			CameraAsset->ParameterDefinitions, DefaultParameterProperties);
+	const UPropertyBag* DefaultParametersStruct = UPropertyBag::GetOrCreateFromDescs(DefaultParameterProperties);
 
+	// Asssign the default values by looking up the same value on the source camera rig.
 	FInstancedPropertyBag DefaultParameters;
-	DefaultParameters.AddProperties(DefaultParameterProperties);
-
-	for (const UCameraRigAsset* CameraRig : AllCameraRigs)
+	DefaultParameters.InitializeFromBagStruct(DefaultParametersStruct);
+	ensure(ParameterDefinitions.Num() == SourceCameraRigs.Num());
+	for (int32 Index = 0; Index < ParameterDefinitions.Num(); ++Index)
 	{
-		CopyDefaultParameterValues(CameraRig, DefaultParameters);
+		const FCameraObjectInterfaceParameterDefinition& ParameterDefinition = ParameterDefinitions[Index];
+		const UCameraRigAsset* SourceCameraRig = SourceCameraRigs[Index];
+		ensure(SourceCameraRig);
+		CopyDefaultParameterValue(SourceCameraRig, ParameterDefinition.ParameterGuid, DefaultParameters);
 	}
 
+	// If the list or types of parameters changed, or if just the default values changed, let's update the asset.
 	if (!DefaultParameters.Identical(&CameraAsset->DefaultParameters, 0))
 	{
 		CameraAsset->Modify();
 		CameraAsset->DefaultParameters = DefaultParameters;
 	}
-
-	// Accumulate all the camera rigs' allocation infos and store that on the asset.
-	FCameraAssetAllocationInfo AllocationInfo;
-
-	for (const UCameraRigAsset* CameraRig : AllCameraRigs)
-	{
-		AllocationInfo.VariableTableInfo.Combine(CameraRig->AllocationInfo.VariableTableInfo);
-		AllocationInfo.ContextDataTableInfo.Combine(CameraRig->AllocationInfo.ContextDataTableInfo);
-	}
-
-	if (AllocationInfo != CameraAsset->AllocationInfo)
-	{
-		CameraAsset->Modify();
-		CameraAsset->AllocationInfo = AllocationInfo;
-	}
 }
 
-void FCameraAssetBuilder::CopyDefaultParameterValues(const UCameraRigAsset* InCameraRig, FInstancedPropertyBag& DefaultParameters)
+void FCameraAssetBuilder::CopyDefaultParameterValue(const UCameraRigAsset* InCameraRig, const FGuid& PropertyID, FInstancedPropertyBag& DefaultParameters)
 {
 	// Our default parameters property bag is an aggregation of various camera rigs' parameters.
 	// Copy their default values over.
@@ -198,38 +301,45 @@ void FCameraAssetBuilder::CopyDefaultParameterValues(const UCameraRigAsset* InCa
 	void* TargetValuesPtr = DefaultParameters.GetMutableValue().GetMemory();
 	const UPropertyBag* SourceValueStruct = SourceValues.GetPropertyBagStruct();
 	const UPropertyBag* TargetValueStruct = DefaultParameters.GetPropertyBagStruct();
-	if (!SourceValuesPtr || !TargetValuesPtr || !SourceValueStruct || !TargetValueStruct)
+	if (!ensure(SourceValuesPtr && TargetValuesPtr && SourceValueStruct && TargetValueStruct))
 	{
-		// We may have empty property bags for camera rigs with no parameters.
+		// We shouldn't have empty property bags if we got this far, since we're supposed to copy at least
+		// one property from one to the other.
 		return;
 	}
 
-	for (const FPropertyBagPropertyDesc& SourcePropertyDesc : SourceValueStruct->GetPropertyDescs())
+	const FPropertyBagPropertyDesc* SourcePropertyDesc = SourceValueStruct->FindPropertyDescByID(PropertyID);
+	const FPropertyBagPropertyDesc* TargetPropertyDesc = TargetValueStruct->FindPropertyDescByID(PropertyID);
+	if (!ensure(SourcePropertyDesc && TargetPropertyDesc))
 	{
-		const FPropertyBagPropertyDesc* TargetPropertyDesc = TargetValueStruct->FindPropertyDescByID(SourcePropertyDesc.ID);
-		if (!ensure(TargetPropertyDesc && TargetPropertyDesc->CachedProperty && SourcePropertyDesc.CachedProperty))
-		{
-			continue;
-		}
-		if (!ensure(SourcePropertyDesc.CompatibleType(*TargetPropertyDesc)))
-		{
-			continue;
-		}
-
-		const void* SourcePtr = SourcePropertyDesc.CachedProperty->ContainerPtrToValuePtr<void>(SourceValuesPtr);
-		void* TargetPtr = TargetPropertyDesc->CachedProperty->ContainerPtrToValuePtr<void>(TargetValuesPtr);
-		SourcePropertyDesc.CachedProperty->CopyCompleteValue(TargetPtr, SourcePtr);
+		// We should have a valid source property here since, in theory, we built our camera asset's parameters
+		// using the referenced camera rigs' built parameter definitions. So their default parameters structs
+		// should be up to date with that and we shouldn't have missing properties.
+		// As for the target property, we should have it since we just made that property bag.
+		return;
 	}
+	if (!ensure(SourcePropertyDesc->CachedProperty && TargetPropertyDesc->CachedProperty))
+	{
+		return;
+	}
+	if (!ensure(SourcePropertyDesc->CompatibleType(*TargetPropertyDesc)))
+	{
+		return;
+	}
+
+	const void* SourcePtr = SourcePropertyDesc->CachedProperty->ContainerPtrToValuePtr<void>(SourceValuesPtr);
+	void* TargetPtr = TargetPropertyDesc->CachedProperty->ContainerPtrToValuePtr<void>(TargetValuesPtr);
+	SourcePropertyDesc->CachedProperty->CopyCompleteValue(TargetPtr, SourcePtr);
 }
 
 void FCameraAssetBuilder::UpdateBuildStatus()
 {
 	ECameraBuildStatus BuildStatus = ECameraBuildStatus::Clean;
-	if (BuildLog.HasErrors())
+	if (BuildContext.BuildLog.HasErrors())
 	{
 		BuildStatus = ECameraBuildStatus::WithErrors;
 	}
-	else if (BuildLog.HasWarnings())
+	else if (BuildContext.BuildLog.HasWarnings())
 	{
 		BuildStatus = ECameraBuildStatus::CleanWithWarnings;
 	}
@@ -238,6 +348,22 @@ void FCameraAssetBuilder::UpdateBuildStatus()
 	CameraAsset->SetBuildStatus(BuildStatus);
 }
 
+bool FCameraAssetBuilder::BuildDefaultInterfaceIfNeeded(UCameraAsset* InCameraAsset)
+{
+	if (InCameraAsset->bNeedsDefaultInterface)
+	{
+		FCameraDirectorRigUsageInfo UsageInfo;
+		Private::GatherCameraAssetRigUsage(InCameraAsset, UsageInfo);
+
+		Private::BuildCameraAssetDefaultParameterDefinitions(InCameraAsset, UsageInfo.CameraRigs);
+		InCameraAsset->bNeedsDefaultInterface = false;
+
+		return true;
+	}
+
+	return false;
+}
+
 }  // namespace UE::Cameras
 
 #undef LOCTEXT_NAMESPACE
diff --git a/Engine/Plugins/Cameras/GameplayCameras/Source/GameplayCameras/Private/Build/CameraNodeHierarchyBuilder.cpp b/Engine/Plugins/Cameras/GameplayCameras/Source/GameplayCameras/Private/Build/CameraNodeHierarchyBuilder.cpp
--- a/Engine/Plugins/Cameras/GameplayCameras/Source/GameplayCameras/Private/Build/CameraNodeHierarchyBuilder.cpp
+++ b/Engine/Plugins/Cameras/GameplayCameras/Source/GameplayCameras/Private/Build/CameraNodeHierarchyBuilder.cpp
@@ -43,8 +43,8 @@ void AddVariableToAllocationInfo(FCameraVariableID VariableID, ECameraVariableTy
 
 }  // namespace Internal
 
-FCameraNodeHierarchyBuilder::FCameraNodeHierarchyBuilder(FCameraBuildLog& InBuildLog, UBaseCameraObject* InCameraObject)
-	: BuildLog(InBuildLog)
+FCameraNodeHierarchyBuilder::FCameraNodeHierarchyBuilder(FCameraBuildContext& InBuildContext, UBaseCameraObject* InCameraObject)
+	: BuildContext(InBuildContext)
 	, CameraObject(InCameraObject)
 {
 	CameraNodeHierarchy.Build(CameraObject);
@@ -54,41 +54,66 @@ void FCameraNodeHierarchyBuilder::PreBuild()
 {
 	for (UCameraNode* CameraNode : CameraNodeHierarchy.GetFlattenedHierarchy())
 	{
-		CameraNode->PreBuild(BuildLog);
+		CameraNode->PreBuild(BuildContext);
 	}
 }
 
 void FCameraNodeHierarchyBuilder::Build()
 {
-	FCameraObjectBuildContext BuildContext(BuildLog);
+	FCameraObjectBuildContext ObjectBuildContext(BuildContext);
 
-	// Build a mock tree of evaluators.
-	FCameraNodeEvaluatorTreeBuildParams BuildParams;
-	BuildParams.RootCameraNode = CameraObject->GetRootNode();
-	FCameraNodeEvaluatorStorage Storage;
-	Storage.BuildEvaluatorTree(BuildParams);
-
-	// Get the size of the evaluators' allocation.
-	Storage.GetAllocationInfo(BuildContext.AllocationInfo.EvaluatorInfo);
+	// Get the size of the evaluators' allocation when cooking. When not cooking, we ignore it because
+	// we don't want to have to re-save assets because someone added a new C++ field to an evaluator.
+	if (BuildContext.IsCooking())
+	{
+		// To get the information we need, we build a mock tree of evaluators and see how much room
+		// that takes.
+		FCameraNodeEvaluatorTreeBuildParams BuildParams;
+		BuildParams.RootCameraNode = CameraObject->GetRootNode();
+		FCameraNodeEvaluatorStorage Storage;
+		Storage.BuildEvaluatorTree(BuildParams);
+
+		Storage.GetAllocationInfo(ObjectBuildContext.AllocationInfo.EvaluatorInfo);
+	}
 
 	// Call Build() on all camera nodes in the hierarchy (detached/orphaned camera nodes don't get called).
 	for (UCameraNode* CameraNode : CameraNodeHierarchy.GetFlattenedHierarchy())
 	{
-		CallBuild(BuildContext, CameraNode);
+		CallBuild(ObjectBuildContext, CameraNode);
 	}
 
 	// Add parameters to the allocation info.
-	BuildParametersAllocationInfo(BuildContext);
+	BuildParametersAllocationInfo(ObjectBuildContext);
 
 	// Set the final allocation info on the camera rig asset.
-	if (CameraObject->AllocationInfo != BuildContext.AllocationInfo)
+	if (CameraObject->AllocationInfo != ObjectBuildContext.AllocationInfo)
 	{
-		CameraObject->Modify();
-		CameraObject->AllocationInfo = BuildContext.AllocationInfo;
+		// Previously we would save the evaluator allocation info, and any referenced camera rigs' infos.
+		// Now not anymore (see above, and see UCameraRigCameraNode::OnBuild... we only do this when cooking now). 
+		// However, we don't want to make all camera rigs modified because this allocation info is now empty or 
+		// something.  So we do this extra check.
+		bool bShouldModify = true;
+		if (!BuildContext.IsCooking())
+		{
+			bShouldModify = false;
+			const FCameraObjectAllocationInfo& CurInfo = CameraObject->AllocationInfo;
+			const FCameraObjectAllocationInfo& NewInfo = ObjectBuildContext.AllocationInfo;
+			if (!CurInfo.VariableTableInfo.Contains(NewInfo.VariableTableInfo) ||
+					!CurInfo.ContextDataTableInfo.Contains(NewInfo.ContextDataTableInfo))
+			{
+				bShouldModify = true;
+			}
+		}
+
+		if (bShouldModify)
+		{
+			CameraObject->Modify();
+			CameraObject->AllocationInfo = ObjectBuildContext.AllocationInfo;
+		}
 	}
 }
 
-void FCameraNodeHierarchyBuilder::CallBuild(FCameraObjectBuildContext& BuildContext, UCameraNode* CameraNode)
+void FCameraNodeHierarchyBuilder::CallBuild(FCameraObjectBuildContext& ObjectBuildContext, UCameraNode* CameraNode)
 {
 	using namespace UE::Cameras::Internal;
 
@@ -96,7 +121,7 @@ void FCameraNodeHierarchyBuilder::CallBuild(FCameraObjectBuildContext& BuildCont
 	// This is only for user-defined variable overrides. We will do the same for exposed camera rig
 	// parameters later, in BuildParametersAllocationInfo.
 	UClass* CameraNodeClass = CameraNode->GetClass();
-	FCameraObjectAllocationInfo& AllocationInfo = BuildContext.AllocationInfo;
+	FCameraObjectAllocationInfo& AllocationInfo = ObjectBuildContext.AllocationInfo;
 	for (TFieldIterator<FProperty> It(CameraNodeClass); It; ++It)
 	{
 		FStructProperty* StructProperty = CastField<FStructProperty>(*It);
@@ -137,10 +162,10 @@ UE_CAMERA_VARIABLE_FOR_ALL_TYPES()
 	}
 
 	// Let the camera node add any custom variables or extra memory.
-	CameraNode->Build(BuildContext);
+	CameraNode->Build(ObjectBuildContext);
 }
 
-void FCameraNodeHierarchyBuilder::BuildParametersAllocationInfo(FCameraObjectBuildContext& BuildContext)
+void FCameraNodeHierarchyBuilder::BuildParametersAllocationInfo(FCameraObjectBuildContext& ObjectBuildContext)
 {
 	// The variables and context data definitions should have already been added by the camera nodes
 	// who have override variable IDs and data IDs set on them. 
@@ -153,7 +178,7 @@ void FCameraNodeHierarchyBuilder::BuildParametersAllocationInfo(FCameraObjectBui
 		}
 
 		FCameraVariableDefinition Definition = BlendableParameter->GetVariableDefinition();
-		BuildContext.AllocationInfo.VariableTableInfo.VariableDefinitions.Add(Definition);
+		ObjectBuildContext.AllocationInfo.VariableTableInfo.VariableDefinitions.Add(Definition);
 	}
 
 	for (const UCameraObjectInterfaceDataParameter* DataParameter : CameraObject->Interface.DataParameters)
@@ -164,7 +189,7 @@ void FCameraNodeHierarchyBuilder::BuildParametersAllocationInfo(FCameraObjectBui
 		}
 
 		FCameraContextDataDefinition Definition = DataParameter->GetDataDefinition();
-		BuildContext.AllocationInfo.ContextDataTableInfo.DataDefinitions.Add(Definition);
+		ObjectBuildContext.AllocationInfo.ContextDataTableInfo.DataDefinitions.Add(Definition);
 	}
 }
 
diff --git a/Engine/Plugins/Cameras/GameplayCameras/Source/GameplayCameras/Private/Build/CameraObjectInterfaceBuilder.cpp b/Engine/Plugins/Cameras/GameplayCameras/Source/GameplayCameras/Private/Build/CameraObjectInterfaceBuilder.cpp
--- a/Engine/Plugins/Cameras/GameplayCameras/Source/GameplayCameras/Private/Build/CameraObjectInterfaceBuilder.cpp
+++ b/Engine/Plugins/Cameras/GameplayCameras/Source/GameplayCameras/Private/Build/CameraObjectInterfaceBuilder.cpp
@@ -2,7 +2,6 @@
 
 #include "Build/CameraObjectInterfaceBuilder.h"
 
-#include "Build/CameraBuildLog.h"
 #include "Core/BaseCameraObject.h"
 #include "Core/CameraNode.h"
 #include "Core/CameraNodeHierarchy.h"
@@ -38,7 +37,7 @@ struct FInterfaceParameterBindingBuilder
 
 	void ReportError(UObject* Object, FText&& ErrorMessage)
 	{
-		Owner.BuildLog.AddMessage(EMessageSeverity::Error, MoveTemp(ErrorMessage));
+		Owner.BuildContext.BuildLog.AddMessage(EMessageSeverity::Error, MoveTemp(ErrorMessage));
 	}
 
 	void SetParameterOverride(
@@ -230,8 +229,8 @@ private:
 
 }  // namespace Internal
 
-FCameraObjectInterfaceBuilder::FCameraObjectInterfaceBuilder(FCameraBuildLog& InBuildLog)
-	: BuildLog(InBuildLog)
+FCameraObjectInterfaceBuilder::FCameraObjectInterfaceBuilder(FCameraBuildContext& InBuildContext)
+	: BuildContext(InBuildContext)
 {
 }
 
@@ -349,7 +348,7 @@ void FCameraObjectInterfaceBuilder::BuildInterfaceParameters()
 		// Basic validations.
 		if (!BlendableParameter)
 		{
-			BuildLog.AddMessage(EMessageSeverity::Warning,
+			BuildContext.BuildLog.AddMessage(EMessageSeverity::Warning,
 					CameraObject,
 					LOCTEXT("InvalidBlendableParameter", "Invalid interface parameter was found and removed."));
 
@@ -361,7 +360,7 @@ void FCameraObjectInterfaceBuilder::BuildInterfaceParameters()
 
 		if (BlendableParameter->InterfaceParameterName.IsEmpty())
 		{
-			BuildLog.AddMessage(EMessageSeverity::Error,
+			BuildContext.BuildLog.AddMessage(EMessageSeverity::Error,
 					BlendableParameter,
 					LOCTEXT(
 						"InvalidBlendableParameterName",
@@ -386,7 +385,7 @@ void FCameraObjectInterfaceBuilder::BuildInterfaceParameters()
 		// Basic validations.
 		if (!DataParameter)
 		{
-			BuildLog.AddMessage(EMessageSeverity::Warning,
+			BuildContext.BuildLog.AddMessage(EMessageSeverity::Warning,
 					CameraObject,
 					LOCTEXT("InvalidDataParameter", "Invalid interface parameter was found and removed."));
 
@@ -398,7 +397,7 @@ void FCameraObjectInterfaceBuilder::BuildInterfaceParameters()
 
 		if (DataParameter->InterfaceParameterName.IsEmpty())
 		{
-			BuildLog.AddMessage(EMessageSeverity::Error,
+			BuildContext.BuildLog.AddMessage(EMessageSeverity::Error,
 					DataParameter,
 					LOCTEXT(
 						"InvalidDataParameterName",
diff --git a/Engine/Plugins/Cameras/GameplayCameras/Source/GameplayCameras/Private/Build/CameraRigAssetBuilder.cpp b/Engine/Plugins/Cameras/GameplayCameras/Source/GameplayCameras/Private/Build/CameraRigAssetBuilder.cpp
--- a/Engine/Plugins/Cameras/GameplayCameras/Source/GameplayCameras/Private/Build/CameraRigAssetBuilder.cpp
+++ b/Engine/Plugins/Cameras/GameplayCameras/Source/GameplayCameras/Private/Build/CameraRigAssetBuilder.cpp
@@ -13,8 +13,8 @@
 namespace UE::Cameras
 {
 
-FCameraRigAssetBuilder::FCameraRigAssetBuilder(FCameraBuildLog& InBuildLog)
-	: BuildLog(InBuildLog)
+FCameraRigAssetBuilder::FCameraRigAssetBuilder(FCameraBuildContext& InBuildContext)
+	: BuildContext(InBuildContext)
 {
 }
 
@@ -27,13 +27,13 @@ void FCameraRigAssetBuilder::BuildCameraRig(UCameraRigAsset* InCameraRig)
 
 	CameraRig = InCameraRig;
 
-	BuildLog.SetLoggingPrefix(CameraRig->GetPathName() + TEXT(": "));
+	BuildContext.BuildLog.SetLoggingPrefix(CameraRig->GetPathName() + TEXT(": "));
 	{
 		BuildCameraRigImpl();
 
 		CameraRig->EventHandlers.Notify(&ICameraRigAssetEventHandler::OnCameraRigBuilt, CameraRig);
 	}
-	BuildLog.SetLoggingPrefix(FString());
+	BuildContext.BuildLog.SetLoggingPrefix(FString());
 	UpdateBuildStatus();
 
 	FGameplayCamerasDelegates::OnCameraRigAssetBuilt().Broadcast(CameraRig);
@@ -43,10 +43,10 @@ void FCameraRigAssetBuilder::BuildCameraRig(UCameraRigAsset* InCameraRig)
 
 void FCameraRigAssetBuilder::BuildCameraRigImpl()
 {
-	FCameraNodeHierarchyBuilder NodeBuilder(BuildLog, CameraRig);
+	FCameraNodeHierarchyBuilder NodeBuilder(BuildContext, CameraRig);
 	NodeBuilder.PreBuild();
 
-	FCameraObjectInterfaceBuilder InterfaceBuilder(BuildLog);
+	FCameraObjectInterfaceBuilder InterfaceBuilder(BuildContext);
 	InterfaceBuilder.BuildInterface(CameraRig, NodeBuilder.GetHierarchy(), true);
 
 	NodeBuilder.Build();
@@ -58,11 +58,11 @@ void FCameraRigAssetBuilder::BuildCameraRigImpl()
 void FCameraRigAssetBuilder::UpdateBuildStatus()
 {
 	ECameraBuildStatus BuildStatus = ECameraBuildStatus::Clean;
-	if (BuildLog.HasErrors())
+	if (BuildContext.BuildLog.HasErrors())
 	{
 		BuildStatus = ECameraBuildStatus::WithErrors;
 	}
-	else if (BuildLog.HasWarnings())
+	else if (BuildContext.BuildLog.HasWarnings())
 	{
 		BuildStatus = ECameraBuildStatus::CleanWithWarnings;
 	}
diff --git a/Engine/Plugins/Cameras/GameplayCameras/Source/GameplayCameras/Private/Build/CameraShakeAssetBuilder.cpp b/Engine/Plugins/Cameras/GameplayCameras/Source/GameplayCameras/Private/Build/CameraShakeAssetBuilder.cpp
--- a/Engine/Plugins/Cameras/GameplayCameras/Source/GameplayCameras/Private/Build/CameraShakeAssetBuilder.cpp
+++ b/Engine/Plugins/Cameras/GameplayCameras/Source/GameplayCameras/Private/Build/CameraShakeAssetBuilder.cpp
@@ -13,8 +13,8 @@
 namespace UE::Cameras
 {
 
-FCameraShakeAssetBuilder::FCameraShakeAssetBuilder(FCameraBuildLog& InBuildLog)
-	: BuildLog(InBuildLog)
+FCameraShakeAssetBuilder::FCameraShakeAssetBuilder(FCameraBuildContext& InBuildContext)
+	: BuildContext(InBuildContext)
 {
 }
 
@@ -27,11 +27,11 @@ void FCameraShakeAssetBuilder::BuildCameraShake(UCameraShakeAsset* InCameraShake
 
 	CameraShake = InCameraShake;
 
-	BuildLog.SetLoggingPrefix(CameraShake->GetPathName() + TEXT(": "));
+	BuildContext.BuildLog.SetLoggingPrefix(CameraShake->GetPathName() + TEXT(": "));
 	{
 		BuildCameraShakeImpl();
 	}
-	BuildLog.SetLoggingPrefix(FString());
+	BuildContext.BuildLog.SetLoggingPrefix(FString());
 	UpdateBuildStatus();
 
 	FGameplayCamerasDelegates::OnCameraShakeAssetBuilt().Broadcast(CameraShake);
@@ -41,10 +41,10 @@ void FCameraShakeAssetBuilder::BuildCameraShake(UCameraShakeAsset* InCameraShake
 
 void FCameraShakeAssetBuilder::BuildCameraShakeImpl()
 {
-	FCameraNodeHierarchyBuilder NodeBuilder(BuildLog, CameraShake);
+	FCameraNodeHierarchyBuilder NodeBuilder(BuildContext, CameraShake);
 	NodeBuilder.PreBuild();
 
-	FCameraObjectInterfaceBuilder InterfaceBuilder(BuildLog);
+	FCameraObjectInterfaceBuilder InterfaceBuilder(BuildContext);
 	InterfaceBuilder.BuildInterface(CameraShake, NodeBuilder.GetHierarchy(), true);
 
 	NodeBuilder.Build();
@@ -56,11 +56,11 @@ void FCameraShakeAssetBuilder::BuildCameraShakeImpl()
 void FCameraShakeAssetBuilder::UpdateBuildStatus()
 {
 	ECameraBuildStatus BuildStatus = ECameraBuildStatus::Clean;
-	if (BuildLog.HasErrors())
+	if (BuildContext.BuildLog.HasErrors())
 	{
 		BuildStatus = ECameraBuildStatus::WithErrors;
 	}
-	else if (BuildLog.HasWarnings())
+	else if (BuildContext.BuildLog.HasWarnings())
 	{
 		BuildStatus = ECameraBuildStatus::CleanWithWarnings;
 	}
diff --git a/Engine/Plugins/Cameras/GameplayCameras/Source/GameplayCameras/Private/Core/BaseCameraObject.cpp b/Engine/Plugins/Cameras/GameplayCameras/Source/GameplayCameras/Private/Core/BaseCameraObject.cpp
--- a/Engine/Plugins/Cameras/GameplayCameras/Source/GameplayCameras/Private/Core/BaseCameraObject.cpp
+++ b/Engine/Plugins/Cameras/GameplayCameras/Source/GameplayCameras/Private/Core/BaseCameraObject.cpp
@@ -100,6 +100,36 @@ void UBaseCameraObject::PostDuplicate(EDuplicateMode::Type DuplicateMode)
 	}
 }
 
+bool UBaseCameraObject::FindParameterDefinitionByName(const FName ParameterName, FCameraObjectInterfaceParameterDefinition& OutParameterDefinition) const
+{
+	const FCameraObjectInterfaceParameterDefinition* FoundItem = ParameterDefinitions.FindByPredicate(
+			[ParameterName](const FCameraObjectInterfaceParameterDefinition& Item)
+			{
+				return Item.ParameterName == ParameterName;
+			});
+	if (FoundItem)
+	{
+		OutParameterDefinition = *FoundItem;
+		return true;
+	}
+	return false;
+}
+
+bool UBaseCameraObject::FindParameterDefinitionByGuid(const FGuid& ParameterGuid, FCameraObjectInterfaceParameterDefinition& OutParameterDefinition) const
+{
+	const FCameraObjectInterfaceParameterDefinition* FoundItem = ParameterDefinitions.FindByPredicate(
+			[&ParameterGuid](const FCameraObjectInterfaceParameterDefinition& Item)
+			{
+				return Item.ParameterGuid == ParameterGuid;
+			});
+	if (FoundItem)
+	{
+		OutParameterDefinition = *FoundItem;
+		return true;
+	}
+	return false;
+}
+
 #if WITH_EDITORONLY_DATA
 
 void UBaseCameraObject::UpgradeInterfaceConnections(IObjectTreeGraphRootObject* RootObject, FName DefaultGraphName)
diff --git a/Engine/Plugins/Cameras/GameplayCameras/Source/GameplayCameras/Private/Core/CameraAsset.cpp b/Engine/Plugins/Cameras/GameplayCameras/Source/GameplayCameras/Private/Core/CameraAsset.cpp
--- a/Engine/Plugins/Cameras/GameplayCameras/Source/GameplayCameras/Private/Core/CameraAsset.cpp
+++ b/Engine/Plugins/Cameras/GameplayCameras/Source/GameplayCameras/Private/Core/CameraAsset.cpp
@@ -3,8 +3,7 @@
 #include "Core/CameraAsset.h"
 
 #include "Build/CameraAssetBuilder.h"
-#include "Build/CameraBuildLog.h"
-#include "Build/CameraObjectInterfaceParameterBuilder.h"
+#include "Build/CameraBuildContext.h"
 #include "Core/CameraDirector.h"
 #include "Core/CameraRigAsset.h"
 #include "Misc/EngineVersionComparison.h"
@@ -13,8 +12,69 @@
 #include "UObject/ObjectRedirector.h"
 #include "UObject/ObjectSaveContext.h"
 
+#if UE_VERSION_OLDER_THAN(5,8,0)
+#include "Build/CameraObjectInterfaceParameterBuilder.h"
+#endif
+
 #include UE_INLINE_GENERATED_CPP_BY_NAME(CameraAsset)
 
+void UCameraAssetInterfaceParameter::PostLoad()
+{
+	if (!Guid.IsValid())
+	{
+		Guid = FGuid::NewGuid();
+	}
+
+	Super::PostLoad();
+}
+
+void UCameraAssetInterfaceParameter::PostInitProperties()
+{
+	Super::PostInitProperties();
+
+	if (!HasAnyFlags(RF_ClassDefaultObject | RF_ArchetypeObject | RF_NeedLoad | RF_WasLoaded) && 
+			!Guid.IsValid())
+	{
+		Guid = FGuid::NewGuid();
+	}
+}
+
+void UCameraAssetInterfaceParameter::PostDuplicate(EDuplicateMode::Type DuplicateMode)
+{
+	Super::PostDuplicate(DuplicateMode);
+
+	if (DuplicateMode == EDuplicateMode::Normal)
+	{
+		Guid = FGuid::NewGuid();
+	}
+}
+
+bool UCameraAssetInterfaceParameter::GetParameterDefinition(FCameraObjectInterfaceParameterDefinition& OutParameterDefinition) const
+{
+	UCameraRigAsset* CameraRig = SourceCameraRig.Get();
+	if (!CameraRig)
+	{
+		return false;
+	}
+
+	const FCameraObjectInterfaceParameterDefinition* SourceParameter = CameraRig->GetParameterDefinitions()
+		.FindByPredicate(
+			[this](const FCameraObjectInterfaceParameterDefinition& Item)
+			{
+				return Item.ParameterName == SourceParameterName;
+			});
+	if (!SourceParameter)
+	{
+		return false;
+	}
+
+	// The parameter definition is the same except that we might expose it under a different name.
+	OutParameterDefinition = *SourceParameter;
+	OutParameterDefinition.ParameterName = FName(InterfaceParameterName);
+
+	return true;
+}
+
 bool operator==(const FCameraAssetAllocationInfo& A, const FCameraAssetAllocationInfo& B)
 {
 	return A.VariableTableInfo == B.VariableTableInfo &&
@@ -114,6 +174,14 @@ void UCameraAsset::PostLoad()
 		}
 	}
 #endif  // WITH_EDITOR
+
+	if (ParameterDefinitions.Num() > 0 && Interface.Parameters.IsEmpty())
+	{
+		// We have saved parameter definitions but they didn't come from building our interface,
+		// so this must be an old asset that pre-dates the interface field. It needs a default
+		// interface to be built for it.
+		bNeedsDefaultInterface = true;
+	}
 }
 
 void UCameraAsset::GetAssetRegistryTags(FAssetRegistryTagsContext Context) const
@@ -161,14 +229,15 @@ void UCameraAsset::BuildCamera()
 
 	FCameraBuildLog BuildLog;
 	BuildLog.SetForwardMessagesToLogging(true);
-	BuildCamera(BuildLog);
+	FCameraBuildContext BuildContext(BuildLog);
+	BuildCamera(BuildContext);
 }
 
-void UCameraAsset::BuildCamera(UE::Cameras::FCameraBuildLog& InBuildLog)
+void UCameraAsset::BuildCamera(UE::Cameras::FCameraBuildContext& InBuildContext)
 {
 	using namespace UE::Cameras;
 
-	FCameraAssetBuilder Builder(InBuildLog);
+	FCameraAssetBuilder Builder(InBuildContext);
 	Builder.BuildCamera(this);
 }
 
diff --git a/Engine/Plugins/Cameras/GameplayCameras/Source/GameplayCameras/Private/Core/CameraContextDataTable.cpp b/Engine/Plugins/Cameras/GameplayCameras/Source/GameplayCameras/Private/Core/CameraContextDataTable.cpp
--- a/Engine/Plugins/Cameras/GameplayCameras/Source/GameplayCameras/Private/Core/CameraContextDataTable.cpp
+++ b/Engine/Plugins/Cameras/GameplayCameras/Source/GameplayCameras/Private/Core/CameraContextDataTable.cpp
@@ -7,6 +7,14 @@
 namespace UE::Cameras
 {
 
+namespace Private
+{
+
+static const uint32 GDefaultCapacity = 64;
+static const uint32 GDefaultAlignment = alignof(void*);
+
+}  // namespace Private
+
 FCameraContextDataTable::FCameraContextDataTable()
 {
 }
@@ -99,8 +107,8 @@ void FCameraContextDataTable::Initialize(const FCameraContextDataTableAllocation
 	EntryLookup.Reset();
 
 	// Compute the total buffer size we need, and create our entries as we go.
-	const uint32 FirstAlignOf = 32u;
 	uint32 TotalSizeOf = 0;
+	uint32 MaxAlignOf = 0;
 	uint32 CurSizeOf, CurAlignOf;
 
 	for (const FCameraContextDataDefinition& DataDefinition : AllocationInfo.DataDefinitions)
@@ -108,29 +116,20 @@ void FCameraContextDataTable::Initialize(const FCameraContextDataTableAllocation
 		GetDataTypeAllocationInfo(DataDefinition.DataType, DataDefinition.DataContainerType, DataDefinition.DataTypeObject, CurSizeOf, CurAlignOf);
 		const uint32 NewEntryOffset = Align(TotalSizeOf, CurAlignOf);
 		TotalSizeOf = NewEntryOffset + CurSizeOf;
+		MaxAlignOf = FMath::Max(MaxAlignOf, CurAlignOf);
 
 		FEntry NewEntry;
-		NewEntry.ID = DataDefinition.DataID;
-		NewEntry.Type = DataDefinition.DataType;
-		NewEntry.ContainerType = DataDefinition.DataContainerType;
-		NewEntry.TypeObject = DataDefinition.DataTypeObject;
+		CreateEntry(DataDefinition, NewEntry);
 		NewEntry.Offset = NewEntryOffset;
-		NewEntry.Flags = EEntryFlags::None;
-		if (DataDefinition.bAutoReset)
-		{
-			NewEntry.Flags |= EEntryFlags::AutoReset;
-		}
-#if WITH_EDITORONLY_DATA
-		NewEntry.DebugName = DataDefinition.DataName;
-#endif
-
 		Entries.Add(NewEntry);
 		EntryLookup.Add(NewEntry.ID, Entries.Num() - 1);
 	}
 
 	// Allocate the memory buffer.
-	Memory = reinterpret_cast<uint8*>(FMemory::Malloc(TotalSizeOf, FirstAlignOf));
+	MaxAlignOf = FMath::Max(Private::GDefaultAlignment, MaxAlignOf);
+	Memory = reinterpret_cast<uint8*>(FMemory::Malloc(TotalSizeOf, MaxAlignOf));
 	Capacity = TotalSizeOf;
+	Alignment = MaxAlignOf;
 	Used = TotalSizeOf;
 
 	// Go back to our entries and initialize each entry to the default value for that data type.
@@ -141,6 +140,59 @@ void FCameraContextDataTable::Initialize(const FCameraContextDataTableAllocation
 	}
 }
 
+bool FCameraContextDataTable::EnsureData(const FCameraContextDataTableAllocationInfo& AllocationInfo)
+{
+	uint32 TotalSizeOf = Used;
+	uint32 MaxAlignOf = 0;
+	uint32 CurSizeOf, CurAlignOf;
+
+	bool bAllExist = true;
+	int32 FirstNewEntry = Entries.Num();
+
+	for (const FCameraContextDataDefinition& DataDefinition : AllocationInfo.DataDefinitions)
+	{
+		if (EntryLookup.Contains(DataDefinition.DataID))
+		{
+			continue;
+		}
+
+		GetDataTypeAllocationInfo(DataDefinition.DataType, DataDefinition.DataContainerType, DataDefinition.DataTypeObject, CurSizeOf, CurAlignOf);
+		const uint32 NewEntryOffset = Align(TotalSizeOf, CurAlignOf);
+		TotalSizeOf = NewEntryOffset + CurSizeOf;
+		MaxAlignOf = FMath::Max(MaxAlignOf, CurAlignOf);
+
+		FEntry NewEntry;
+		CreateEntry(DataDefinition, NewEntry);
+		NewEntry.Offset = NewEntryOffset;
+		Entries.Add(NewEntry);
+		EntryLookup.Add(NewEntry.ID, Entries.Num() - 1);
+
+		bAllExist = false;
+	}
+
+	if (!bAllExist)
+	{
+		// Alignment may not have been initialized yet.
+		Alignment = FMath::Max(Alignment, MaxAlignOf);
+
+		if (TotalSizeOf > Capacity)
+		{
+			ReallocateBuffer(TotalSizeOf);
+		}
+
+		Used = TotalSizeOf;
+
+		for (int32 Index = FirstNewEntry; Index < Entries.Num(); ++Index)
+		{
+			const FEntry& Entry = Entries[Index];
+			uint8* DataPtr = Memory + Entry.Offset;
+			ConstructDataValue(Entry.Type, Entry.ContainerType, Entry.TypeObject, DataPtr);
+		}
+	}
+
+	return bAllExist;
+}
+
 void FCameraContextDataTable::AddData(const FCameraContextDataDefinition& DataDefinition)
 {
 	if (!ensure(!EntryLookup.Contains(DataDefinition.DataID)))
@@ -156,6 +208,9 @@ void FCameraContextDataTable::AddData(const FCameraContextDataDefinition& DataDe
 
 	if (NewUsed > Capacity)
 	{
+		// Alignment may not have been initialized yet.
+		Alignment = FMath::Max(AlignOf, Private::GDefaultAlignment);
+
 		ReallocateBuffer(NewUsed);
 
 		DataPtr = Align(Memory + Used, AlignOf);
@@ -164,24 +219,29 @@ void FCameraContextDataTable::AddData(const FCameraContextDataDefinition& DataDe
 	Used = NewUsed;
 
 	FEntry NewEntry;
-	NewEntry.ID = DataDefinition.DataID;
-	NewEntry.Type = DataDefinition.DataType;
-	NewEntry.ContainerType = DataDefinition.DataContainerType;
-	NewEntry.TypeObject = DataDefinition.DataTypeObject;
+	CreateEntry(DataDefinition, NewEntry);
 	NewEntry.Offset = DataPtr - Memory;
-	NewEntry.Flags = EEntryFlags::None;
+	Entries.Add(NewEntry);
+	EntryLookup.Add(DataDefinition.DataID, Entries.Num() - 1);
+
+	ConstructDataValue(NewEntry.Type, NewEntry.ContainerType, NewEntry.TypeObject, Memory + NewEntry.Offset);
+}
+
+void FCameraContextDataTable::CreateEntry(const FCameraContextDataDefinition& DataDefinition, FEntry& OutEntry)
+{
+	OutEntry.ID = DataDefinition.DataID;
+	OutEntry.Type = DataDefinition.DataType;
+	OutEntry.ContainerType = DataDefinition.DataContainerType;
+	OutEntry.TypeObject = DataDefinition.DataTypeObject;
+	OutEntry.Offset = 0;
+	OutEntry.Flags = EEntryFlags::None;
 	if (DataDefinition.bAutoReset)
 	{
-		NewEntry.Flags |= EEntryFlags::AutoReset;
+		OutEntry.Flags |= EEntryFlags::AutoReset;
 	}
 #if WITH_EDITORONLY_DATA
-	NewEntry.DebugName = DataDefinition.DataName;
+	OutEntry.DebugName = DataDefinition.DataName;
 #endif
-	
-	Entries.Add(NewEntry);
-	EntryLookup.Add(DataDefinition.DataID, Entries.Num() - 1);
-
-	ConstructDataValue(NewEntry.Type, NewEntry.ContainerType, NewEntry.TypeObject, Memory + NewEntry.Offset);
 }
 
 bool FCameraContextDataTable::GetDataTypeAllocationInfo(ECameraContextDataType DataType, const UObject* DataTypeObject, uint32& OutSizeOf, uint32& OutAlignOf)
@@ -350,17 +410,14 @@ bool FCameraContextDataTable::DestroyDataValue(ECameraContextDataType DataType,
 
 void FCameraContextDataTable::ReallocateBuffer(uint32 MinRequired)
 {
-	static const uint32 DefaultCapacity = 64u;
-	static const uint32 DefaultAlignment = 32u;
-
-	uint32 NewCapacity = Capacity <= 0 ? DefaultCapacity : Capacity * 2;
+	uint32 NewCapacity = Capacity <= 0 ? Private::GDefaultCapacity : Capacity * 2;
 	if (MinRequired > 0)
 	{
 		NewCapacity = FMath::Max(NewCapacity, MinRequired);
 	}
 
 	uint8* OldMemory = Memory;
-	uint8* NewMemory = reinterpret_cast<uint8*>(FMemory::Malloc(NewCapacity, DefaultAlignment));
+	uint8* NewMemory = reinterpret_cast<uint8*>(FMemory::Malloc(NewCapacity, Alignment));
 
 	if (OldMemory)
 	{
diff --git a/Engine/Plugins/Cameras/GameplayCameras/Source/GameplayCameras/Private/Core/CameraContextDataTableAllocationInfo.cpp b/Engine/Plugins/Cameras/GameplayCameras/Source/GameplayCameras/Private/Core/CameraContextDataTableAllocationInfo.cpp
--- a/Engine/Plugins/Cameras/GameplayCameras/Source/GameplayCameras/Private/Core/CameraContextDataTableAllocationInfo.cpp
+++ b/Engine/Plugins/Cameras/GameplayCameras/Source/GameplayCameras/Private/Core/CameraContextDataTableAllocationInfo.cpp
@@ -4,18 +4,30 @@
 
 #include UE_INLINE_GENERATED_CPP_BY_NAME(CameraContextDataTableAllocationInfo)
 
-void FCameraContextDataTableAllocationInfo::Combine(const FCameraContextDataTableAllocationInfo& OtherInfo)
+namespace UE::Cameras::Private
+{
+
+void GetKnownContextDataDefinitions(
+		const FCameraContextDataTableAllocationInfo& AllocationInfo,
+		TMap<FCameraContextDataID, int32>& OutKnownIDs)
 {
-	TMap<FCameraContextDataID, int32> KnownNames;
-	for (auto It = DataDefinitions.CreateConstIterator(); It; ++It)
+	for (auto It = AllocationInfo.DataDefinitions.CreateConstIterator(); It; ++It)
 	{
-		const FCameraContextDataDefinition& DataDefinition(*It);
-		KnownNames.Add(DataDefinition.DataID, It.GetIndex());
+		const FCameraContextDataDefinition& ContextDataDefinition(*It);
+		OutKnownIDs.Add(ContextDataDefinition.DataID, It.GetIndex());
 	}
+}
+
+}  // namespace UE::Cameras::Private
+
+void FCameraContextDataTableAllocationInfo::Combine(const FCameraContextDataTableAllocationInfo& OtherInfo)
+{
+	TMap<FCameraContextDataID, int32> KnownIDs;
+	UE::Cameras::Private::GetKnownContextDataDefinitions(*this, KnownIDs);
 
 	for (const FCameraContextDataDefinition& OtherDataDefinition : OtherInfo.DataDefinitions)
 	{
-		const int32 KnownIndex = KnownNames.FindRef(OtherDataDefinition.DataID, INDEX_NONE);
+		const int32 KnownIndex = KnownIDs.FindRef(OtherDataDefinition.DataID, INDEX_NONE);
 		if (KnownIndex == INDEX_NONE)
 		{
 			DataDefinitions.Add(OtherDataDefinition);
@@ -28,3 +40,20 @@ void FCameraContextDataTableAllocationInfo::Combine(const FCameraContextDataTabl
 	}
 }
 
+bool FCameraContextDataTableAllocationInfo::Contains(const FCameraContextDataTableAllocationInfo& OtherInfo) const
+{
+	TMap<FCameraContextDataID, int32> KnownIDs;
+	UE::Cameras::Private::GetKnownContextDataDefinitions(*this, KnownIDs);
+
+	for (const FCameraContextDataDefinition& OtherDataDefinition : OtherInfo.DataDefinitions)
+	{
+		const int32 KnownIndex = KnownIDs.FindRef(OtherDataDefinition.DataID, INDEX_NONE);
+		if (KnownIndex == INDEX_NONE)
+		{
+			return false;
+		}
+	}
+
+	return true;
+}
+
diff --git a/Engine/Plugins/Cameras/GameplayCameras/Source/GameplayCameras/Private/Core/CameraDirector.cpp b/Engine/Plugins/Cameras/GameplayCameras/Source/GameplayCameras/Private/Core/CameraDirector.cpp
--- a/Engine/Plugins/Cameras/GameplayCameras/Source/GameplayCameras/Private/Core/CameraDirector.cpp
+++ b/Engine/Plugins/Cameras/GameplayCameras/Source/GameplayCameras/Private/Core/CameraDirector.cpp
@@ -29,14 +29,32 @@ FCameraDirectorEvaluatorPtr UCameraDirector::BuildEvaluator(FCameraDirectorEvalu
 	return NewEvaluator;
 }
 
-void UCameraDirector::BuildCameraDirector(UE::Cameras::FCameraBuildLog& BuildLog)
+void UCameraDirector::BuildCameraDirector(UE::Cameras::FCameraBuildContext& BuildContext)
 {
-	OnBuildCameraDirector(BuildLog);
+	OnBuildCameraDirector(BuildContext);
 }
 
 void UCameraDirector::GatherRigUsageInfo(FCameraDirectorRigUsageInfo& UsageInfo) const
 {
 	OnGatherRigUsageInfo(UsageInfo);
+
+	// Add the rigs referenced by the proxy table.
+	for (const FCameraRigProxyRedirectTableEntry& Entry : CameraRigProxyRedirectTable.Entries)
+	{
+		if (Entry.CameraRigProxy)
+		{
+			UsageInfo.CameraRigProxies.AddUnique(Entry.CameraRigProxy);
+		}
+		if (Entry.CameraRig)
+		{
+			UsageInfo.CameraRigs.AddUnique(Entry.CameraRig);
+		}
+	}
+
+	// Some directors might add CDOs for some reason (for instance if they use a reference collector).
+	// We need to remove them and leave them untouched.
+	UsageInfo.CameraRigs.Remove(GetMutableDefault<UCameraRigAsset>());
+	UsageInfo.CameraRigProxies.Remove(GetMutableDefault<UCameraRigProxyAsset>());
 }
 
 void UCameraDirector::ExtendAssetRegistryTags(FAssetRegistryTagsContext Context) const
@@ -67,7 +85,7 @@ void UCameraDirector::ExtendAssetRegistryTags(FAssetRegistryTagsContext Context)
 	{
 		FAssetRegistryTag NumCameraRigProxiesTag;
 		NumCameraRigProxiesTag.Name = TEXT("NumUsedCameraRigProxies");
-		NumCameraRigProxiesTag.Value = LexToString(UsageInfo.CameraRigs.Num());
+		NumCameraRigProxiesTag.Value = LexToString(UsageInfo.CameraRigProxies.Num());
 		Context.AddTag(NumCameraRigProxiesTag);
 
 		TStringBuilder<256> CameraRigProxyListBuilder;
diff --git a/Engine/Plugins/Cameras/GameplayCameras/Source/GameplayCameras/Private/Core/CameraNode.cpp b/Engine/Plugins/Cameras/GameplayCameras/Source/GameplayCameras/Private/Core/CameraNode.cpp
--- a/Engine/Plugins/Cameras/GameplayCameras/Source/GameplayCameras/Private/Core/CameraNode.cpp
+++ b/Engine/Plugins/Cameras/GameplayCameras/Source/GameplayCameras/Private/Core/CameraNode.cpp
@@ -60,9 +60,9 @@ FCameraNodeChildrenView UCameraNode::GetChildren()
 	return ChildrenView;
 }
 
-void UCameraNode::PreBuild(FCameraBuildLog& BuildLog)
+void UCameraNode::PreBuild(FCameraBuildContext& BuildContext)
 {
-	OnPreBuild(BuildLog);
+	OnPreBuild(BuildContext);
 }
 
 void UCameraNode::Build(FCameraObjectBuildContext& BuildContext)
diff --git a/Engine/Plugins/Cameras/GameplayCameras/Source/GameplayCameras/Private/Core/CameraRigAsset.cpp b/Engine/Plugins/Cameras/GameplayCameras/Source/GameplayCameras/Private/Core/CameraRigAsset.cpp
--- a/Engine/Plugins/Cameras/GameplayCameras/Source/GameplayCameras/Private/Core/CameraRigAsset.cpp
+++ b/Engine/Plugins/Cameras/GameplayCameras/Source/GameplayCameras/Private/Core/CameraRigAsset.cpp
@@ -2,6 +2,7 @@
 
 #include "Core/CameraRigAsset.h"
 
+#include "Build/CameraBuildContext.h"
 #include "Build/CameraBuildLog.h"
 #include "Build/CameraRigAssetBuilder.h"
 #include "Core/CameraNode.h"
@@ -85,14 +86,15 @@ void UCameraRigAsset::BuildCameraRig()
 
 	FCameraBuildLog BuildLog;
 	BuildLog.SetForwardMessagesToLogging(true);
-	BuildCameraRig(BuildLog);
+	FCameraBuildContext BuildContext(BuildLog);
+	BuildCameraRig(BuildContext);
 }
 
-void UCameraRigAsset::BuildCameraRig(UE::Cameras::FCameraBuildLog& InBuildLog)
+void UCameraRigAsset::BuildCameraRig(UE::Cameras::FCameraBuildContext& InBuildContext)
 {
 	using namespace UE::Cameras;
 
-	FCameraRigAssetBuilder Builder(InBuildLog);
+	FCameraRigAssetBuilder Builder(InBuildContext);
 	Builder.BuildCameraRig(this);
 }
 
diff --git a/Engine/Plugins/Cameras/GameplayCameras/Source/GameplayCameras/Private/Core/CameraRigAssetReference.cpp b/Engine/Plugins/Cameras/GameplayCameras/Source/GameplayCameras/Private/Core/CameraRigAssetReference.cpp
--- a/Engine/Plugins/Cameras/GameplayCameras/Source/GameplayCameras/Private/Core/CameraRigAssetReference.cpp
+++ b/Engine/Plugins/Cameras/GameplayCameras/Source/GameplayCameras/Private/Core/CameraRigAssetReference.cpp
@@ -27,6 +27,16 @@ const UBaseCameraObject* FCameraRigAssetReference::GetCameraObject() const
 	return CameraRig;
 }
 
+void FCameraRigAssetReference::EnsureAllocationInfo(UE::Cameras::FCameraNodeEvaluationResult& OutResult) const
+{
+	if (CameraRig)
+	{
+		const FCameraObjectAllocationInfo& AllocationInfo = CameraRig->AllocationInfo;
+		OutResult.VariableTable.EnsureVariables(AllocationInfo.VariableTableInfo);
+		OutResult.ContextDataTable.EnsureData(AllocationInfo.ContextDataTableInfo);
+	}
+}
+
 void FCameraRigAssetReference::ApplyParameterOverrides(UE::Cameras::FCameraVariableTable& OutVariableTable, bool bDrivenOnly) const
 {
 	ApplyParameterOverridesImpl(&OutVariableTable, nullptr, bDrivenOnly, false);
diff --git a/Engine/Plugins/Cameras/GameplayCameras/Source/GameplayCameras/Private/Core/CameraShakeAsset.cpp b/Engine/Plugins/Cameras/GameplayCameras/Source/GameplayCameras/Private/Core/CameraShakeAsset.cpp
--- a/Engine/Plugins/Cameras/GameplayCameras/Source/GameplayCameras/Private/Core/CameraShakeAsset.cpp
+++ b/Engine/Plugins/Cameras/GameplayCameras/Source/GameplayCameras/Private/Core/CameraShakeAsset.cpp
@@ -2,7 +2,7 @@
 
 #include "Core/CameraShakeAsset.h"
 
-#include "Build/CameraBuildLog.h"
+#include "Build/CameraBuildContext.h"
 #include "Build/CameraShakeAssetBuilder.h"
 #include "Core/ShakeCameraNode.h"  // IWYU pragma: keep
 #include "Misc/EngineVersionComparison.h"
@@ -17,14 +17,15 @@ void UCameraShakeAsset::BuildCameraShake()
 
 	FCameraBuildLog BuildLog;
 	BuildLog.SetForwardMessagesToLogging(true);
-	BuildCameraShake(BuildLog);
+	FCameraBuildContext BuildContext(BuildLog);
+	BuildCameraShake(BuildContext);
 }
 
-void UCameraShakeAsset::BuildCameraShake(UE::Cameras::FCameraBuildLog& InBuildLog)
+void UCameraShakeAsset::BuildCameraShake(UE::Cameras::FCameraBuildContext& InBuildContext)
 {
 	using namespace UE::Cameras;
 
-	FCameraShakeAssetBuilder Builder(InBuildLog);
+	FCameraShakeAssetBuilder Builder(InBuildContext);
 	Builder.BuildCameraShake(this);
 }
 
diff --git a/Engine/Plugins/Cameras/GameplayCameras/Source/GameplayCameras/Private/Core/CameraShakeAssetReference.cpp b/Engine/Plugins/Cameras/GameplayCameras/Source/GameplayCameras/Private/Core/CameraShakeAssetReference.cpp
--- a/Engine/Plugins/Cameras/GameplayCameras/Source/GameplayCameras/Private/Core/CameraShakeAssetReference.cpp
+++ b/Engine/Plugins/Cameras/GameplayCameras/Source/GameplayCameras/Private/Core/CameraShakeAssetReference.cpp
@@ -22,6 +22,16 @@ const UBaseCameraObject* FCameraShakeAssetReference::GetCameraObject() const
 	return CameraShake;
 }
 
+void FCameraShakeAssetReference::EnsureAllocationInfo(UE::Cameras::FCameraNodeEvaluationResult& OutResult) const
+{
+	if (CameraShake)
+	{
+		const FCameraObjectAllocationInfo& AllocationInfo = CameraShake->AllocationInfo;
+		OutResult.VariableTable.EnsureVariables(AllocationInfo.VariableTableInfo);
+		OutResult.ContextDataTable.EnsureData(AllocationInfo.ContextDataTableInfo);
+	}
+}
+
 void FCameraShakeAssetReference::ApplyParameterOverrides(UE::Cameras::FCameraVariableTable& OutVariableTable, bool bDrivenOnly) const
 {
 	ApplyParameterOverridesImpl(&OutVariableTable, nullptr, bDrivenOnly, false);
diff --git a/Engine/Plugins/Cameras/GameplayCameras/Source/GameplayCameras/Private/Core/CameraVariableTable.cpp b/Engine/Plugins/Cameras/GameplayCameras/Source/GameplayCameras/Private/Core/CameraVariableTable.cpp
--- a/Engine/Plugins/Cameras/GameplayCameras/Source/GameplayCameras/Private/Core/CameraVariableTable.cpp
+++ b/Engine/Plugins/Cameras/GameplayCameras/Source/GameplayCameras/Private/Core/CameraVariableTable.cpp
@@ -12,6 +12,9 @@ namespace UE::Cameras
 namespace Private
 {
 
+static const uint32 GDefaultCapacity = 64;
+static const uint32 GDefaultAlignment = alignof(void*);
+
 static const FString GUnavailableVariableDebugName(TEXT("<no debug info>"));
 
 bool IsVariableInMask(FCameraVariableID VariableID, const FCameraVariableTableFlags* InMask, bool bInvertMask)
@@ -38,10 +41,12 @@ FCameraVariableTable::FCameraVariableTable(FCameraVariableTable&& Other)
 	, EntryLookup(MoveTemp(Other.EntryLookup))
 	, Memory(Other.Memory)
 	, Capacity(Other.Capacity)
+	, Alignment(Other.Alignment)
 	, Used(Other.Used)
 {
 	Other.Memory = nullptr;
 	Other.Capacity = 0;
+	Other.Alignment = 0;
 	Other.Used = 0;
 }
 
@@ -51,10 +56,12 @@ FCameraVariableTable& FCameraVariableTable::operator=(FCameraVariableTable&& Oth
 	EntryLookup = MoveTemp(Other.EntryLookup);
 	Memory = Other.Memory;
 	Capacity = Other.Capacity;
+	Alignment = Other.Alignment;
 	Used = Other.Used;
 
 	Other.Memory = nullptr;
 	Other.Capacity = 0;
+	Other.Alignment = 0;
 	Other.Used = 0;
 
 	return *this;
@@ -66,7 +73,7 @@ FCameraVariableTable::~FCameraVariableTable()
 	{
 		FMemory::Free(Memory);
 		Memory = nullptr;
-		Capacity = Used = 0;
+		Capacity = Alignment = Used = 0;
 	}
 }
 
@@ -93,39 +100,53 @@ void FCameraVariableTable::Initialize(const FCameraVariableTableAllocationInfo&
 		MaxAlignOf = FMath::Max(MaxAlignOf, CurAlignOf);
 
 		FEntry NewEntry;
-		NewEntry.ID = VariableDefinition.VariableID;
-		NewEntry.Type = VariableDefinition.VariableType;
-		NewEntry.StructType = VariableDefinition.BlendableStructType;
+		CreateEntry(VariableDefinition, NewEntry);
 		NewEntry.Offset = NewEntryOffset;
-		NewEntry.Flags = EEntryFlags::None;
-		if (VariableDefinition.bIsPrivate)
-		{
-			NewEntry.Flags |= EEntryFlags::Private;
-		}
-		if (VariableDefinition.bIsInput)
-		{
-			NewEntry.Flags |= EEntryFlags::Input;
-		}
-		if (VariableDefinition.bAutoReset)
-		{
-			NewEntry.Flags |= EEntryFlags::AutoReset;
-		}
-#if WITH_EDITORONLY_DATA
-		NewEntry.DebugName = VariableDefinition.VariableName;
-#endif
 		Entries.Add(NewEntry);
 		EntryLookup.Add(NewEntry.ID, Entries.Num() - 1);
 	}
 
 	// Allocate the memory buffer.
-	MaxAlignOf = FMath::Max(32u, MaxAlignOf);
+	MaxAlignOf = FMath::Max(Private::GDefaultAlignment, MaxAlignOf);
 	Memory = reinterpret_cast<uint8*>(FMemory::Malloc(TotalSizeOf, MaxAlignOf));
 	Capacity = TotalSizeOf;
+	Alignment = MaxAlignOf;
 	Used = TotalSizeOf;
 
 	// Go back to our entries and initialize each entry to the default value for that variable type.
-	for (const FEntry& Entry : Entries)
+	InitializeEntries();
+}
+
+void FCameraVariableTable::CreateEntry(const FCameraVariableDefinition& VariableDefinition, FEntry& OutEntry)
+{
+	OutEntry.ID = VariableDefinition.VariableID;
+	OutEntry.Type = VariableDefinition.VariableType;
+	OutEntry.StructType = VariableDefinition.BlendableStructType;
+	OutEntry.Offset = 0;
+	OutEntry.Flags = EEntryFlags::None;
+	if (VariableDefinition.bIsPrivate)
+	{
+		OutEntry.Flags |= EEntryFlags::Private;
+	}
+	if (VariableDefinition.bIsInput)
+	{
+		OutEntry.Flags |= EEntryFlags::Input;
+	}
+	if (VariableDefinition.bAutoReset)
 	{
+		OutEntry.Flags |= EEntryFlags::AutoReset;
+	}
+#if WITH_EDITORONLY_DATA
+	OutEntry.DebugName = VariableDefinition.VariableName;
+#endif
+}
+
+void FCameraVariableTable::InitializeEntries(int32 StartIndex, int32 EndIndex)
+{
+	EndIndex = (EndIndex >= 0 ? EndIndex : Entries.Num());
+	for (int32 Index = StartIndex; Index < EndIndex; ++Index)
+	{
+		const FEntry& Entry = Entries[Index];
 		uint8* ValuePtr = Memory + Entry.Offset;
 		switch (Entry.Type)
 		{
@@ -148,6 +169,54 @@ UE_CAMERA_VARIABLE_FOR_ALL_TYPES()
 	}
 }
 
+bool FCameraVariableTable::EnsureVariables(const FCameraVariableTableAllocationInfo& AllocationInfo)
+{
+	uint32 TotalSizeOf = Used;
+	uint32 MaxAlignOf = 0;
+	uint32 CurSizeOf, CurAlignOf;
+
+	bool bAllExist = true;
+	int32 FirstNewEntry = Entries.Num();
+
+	for (const FCameraVariableDefinition& VariableDefinition : AllocationInfo.VariableDefinitions)
+	{
+		if (EntryLookup.Contains(VariableDefinition.VariableID))
+		{
+			continue;
+		}
+
+		GetVariableTypeAllocationInfo(VariableDefinition.VariableType, VariableDefinition.BlendableStructType, CurSizeOf, CurAlignOf);
+		const uint32 NewEntryOffset = Align(TotalSizeOf, CurAlignOf);
+		TotalSizeOf = NewEntryOffset + CurSizeOf;
+		MaxAlignOf = FMath::Max(MaxAlignOf, CurAlignOf);
+
+		FEntry NewEntry;
+		CreateEntry(VariableDefinition, NewEntry);
+		NewEntry.Offset = NewEntryOffset;
+		Entries.Add(NewEntry);
+		EntryLookup.Add(NewEntry.ID, Entries.Num() - 1);
+
+		bAllExist = false;
+	}
+
+	if (!bAllExist)
+	{
+		// Alignment may not have been initialized yet.
+		Alignment = FMath::Max(Alignment, MaxAlignOf);
+
+		if (TotalSizeOf > Capacity)
+		{
+			ReallocateBuffer(TotalSizeOf);
+		}
+
+		Used = TotalSizeOf;
+
+		InitializeEntries(FirstNewEntry);
+	}
+
+	return bAllExist;
+}
+
 void FCameraVariableTable::CacheBlendableStructs()
 {
 	if (bCachedBlendableStructs)
@@ -184,6 +253,9 @@ void FCameraVariableTable::AddVariable(const FCameraVariableDefinition& Variable
 
 	if (NewUsed > Capacity)
 	{
+		// Alignment may not have been initialized yet.
+		Alignment = FMath::Max(AlignOf, Private::GDefaultAlignment);
+
 		ReallocateBuffer(NewUsed);
 
 		VariablePtr = Align(Memory + Used, AlignOf);
@@ -192,44 +264,24 @@ void FCameraVariableTable::AddVariable(const FCameraVariableDefinition& Variable
 	Used = NewUsed;
 
 	FEntry NewEntry;
-	NewEntry.ID = VariableDefinition.VariableID;
-	NewEntry.Type = VariableDefinition.VariableType;
-	NewEntry.StructType = VariableDefinition.BlendableStructType;
+	CreateEntry(VariableDefinition, NewEntry);
 	NewEntry.Offset = VariablePtr - Memory;
-	NewEntry.Flags = EEntryFlags::None;
-	if (VariableDefinition.bIsPrivate)
-	{
-		NewEntry.Flags |= EEntryFlags::Private;
-	}
-	if (VariableDefinition.bIsInput)
-	{
-		NewEntry.Flags |= EEntryFlags::Input;
-	}
-	if (VariableDefinition.bAutoReset)
-	{
-		NewEntry.Flags |= EEntryFlags::AutoReset;
-	}
-#if WITH_EDITORONLY_DATA
-	NewEntry.DebugName = VariableDefinition.VariableName;
-#endif
-
 	Entries.Add(NewEntry);
 	EntryLookup.Add(VariableDefinition.VariableID, Entries.Num() - 1);
+
+	InitializeEntries(Entries.Num() - 1);
 }
 
 void FCameraVariableTable::ReallocateBuffer(uint32 MinRequired)
 {
-	static const uint32 DefaultCapacity = 64;
-	static const uint32 DefaultAlignment = 32;
-
-	uint32 NewCapacity = Capacity <= 0 ? DefaultCapacity : Capacity * 2;
+	uint32 NewCapacity = Capacity <= 0 ? Private::GDefaultCapacity : Capacity * 2;
 	if (MinRequired > 0)
 	{
 		NewCapacity = FMath::Max(NewCapacity, MinRequired);
 	}
 
 	uint8* OldMemory = Memory;
-	uint8* NewMemory = reinterpret_cast<uint8*>(FMemory::Malloc(NewCapacity, DefaultAlignment));
+	uint8* NewMemory = reinterpret_cast<uint8*>(FMemory::Malloc(NewCapacity, Alignment));
 
 	if (OldMemory)
 	{
diff --git a/Engine/Plugins/Cameras/GameplayCameras/Source/GameplayCameras/Private/Core/CameraVariableTableAllocationInfo.cpp b/Engine/Plugins/Cameras/GameplayCameras/Source/GameplayCameras/Private/Core/CameraVariableTableAllocationInfo.cpp
--- a/Engine/Plugins/Cameras/GameplayCameras/Source/GameplayCameras/Private/Core/CameraVariableTableAllocationInfo.cpp
+++ b/Engine/Plugins/Cameras/GameplayCameras/Source/GameplayCameras/Private/Core/CameraVariableTableAllocationInfo.cpp
@@ -4,14 +4,26 @@
 
 #include UE_INLINE_GENERATED_CPP_BY_NAME(CameraVariableTableAllocationInfo)
 
-void FCameraVariableTableAllocationInfo::Combine(const FCameraVariableTableAllocationInfo& OtherInfo)
+namespace UE::Cameras::Private
 {
-	TMap<FCameraVariableID, int32> KnownIDs;
-	for (auto It = VariableDefinitions.CreateConstIterator(); It; ++It)
+
+void GetKnownVariableDefinitions(
+		const FCameraVariableTableAllocationInfo& AllocationInfo,
+		TMap<FCameraVariableID, int32>& OutKnownIDs)
+{
+	for (auto It = AllocationInfo.VariableDefinitions.CreateConstIterator(); It; ++It)
 	{
 		const FCameraVariableDefinition& VariableDefinition(*It);
-		KnownIDs.Add(VariableDefinition.VariableID, It.GetIndex());
+		OutKnownIDs.Add(VariableDefinition.VariableID, It.GetIndex());
 	}
+}
+
+}  // namespace UE::Cameras::Private
+
+void FCameraVariableTableAllocationInfo::Combine(const FCameraVariableTableAllocationInfo& OtherInfo)
+{
+	TMap<FCameraVariableID, int32> KnownIDs;
+	UE::Cameras::Private::GetKnownVariableDefinitions(*this, KnownIDs);
 
 	for (const FCameraVariableDefinition& OtherVariableDefinition : OtherInfo.VariableDefinitions)
 	{
@@ -28,3 +40,20 @@ void FCameraVariableTableAllocationInfo::Combine(const FCameraVariableTableAlloc
 	}
 }
 
+bool FCameraVariableTableAllocationInfo::Contains(const FCameraVariableTableAllocationInfo& OtherInfo) const
+{
+	TMap<FCameraVariableID, int32> KnownIDs;
+	UE::Cameras::Private::GetKnownVariableDefinitions(*this, KnownIDs);
+
+	for (const FCameraVariableDefinition& OtherVariableDefinition : OtherInfo.VariableDefinitions)
+	{
+		const int32 KnownIndex = KnownIDs.FindRef(OtherVariableDefinition.VariableID, INDEX_NONE);
+		if (KnownIndex == INDEX_NONE)
+		{
+			return false;
+		}
+	}
+
+	return true;
+}
+
diff --git a/Engine/Plugins/Cameras/GameplayCameras/Source/GameplayCameras/Private/Directors/BlueprintCameraDirector.cpp b/Engine/Plugins/Cameras/GameplayCameras/Source/GameplayCameras/Private/Directors/BlueprintCameraDirector.cpp
--- a/Engine/Plugins/Cameras/GameplayCameras/Source/GameplayCameras/Private/Directors/BlueprintCameraDirector.cpp
+++ b/Engine/Plugins/Cameras/GameplayCameras/Source/GameplayCameras/Private/Directors/BlueprintCameraDirector.cpp
@@ -2,7 +2,7 @@
 
 #include "Directors/BlueprintCameraDirector.h"
 
-#include "Build/CameraBuildLog.h"
+#include "Build/CameraBuildContext.h"
 #include "Components/ActorComponent.h"
 #include "Core/CameraAsset.h"
 #include "Core/CameraDirectorEvaluator.h"
@@ -501,14 +501,14 @@ FCameraDirectorEvaluatorPtr UBlueprintCameraDirector::OnBuildEvaluator(FCameraDi
 	return Builder.BuildEvaluator<FBlueprintCameraDirectorEvaluator>();
 }
 
-void UBlueprintCameraDirector::OnBuildCameraDirector(UE::Cameras::FCameraBuildLog& BuildLog)
+void UBlueprintCameraDirector::OnBuildCameraDirector(UE::Cameras::FCameraBuildContext& BuildContext)
 {
 	using namespace UE::Cameras;
 
 	// Check that a camera director evaluator Blueprint was specified.
 	if (!CameraDirectorEvaluatorClass)
 	{
-		BuildLog.AddMessage(EMessageSeverity::Error, LOCTEXT("MissingBlueprintClass", "No evaluator Blueprint class is set."));
+		BuildContext.BuildLog.AddMessage(EMessageSeverity::Error, LOCTEXT("MissingBlueprintClass", "No evaluator Blueprint class is set."));
 		return;
 	}
 }
diff --git a/Engine/Plugins/Cameras/GameplayCameras/Source/GameplayCameras/Private/Directors/SingleCameraDirector.cpp b/Engine/Plugins/Cameras/GameplayCameras/Source/GameplayCameras/Private/Directors/SingleCameraDirector.cpp
--- a/Engine/Plugins/Cameras/GameplayCameras/Source/GameplayCameras/Private/Directors/SingleCameraDirector.cpp
+++ b/Engine/Plugins/Cameras/GameplayCameras/Source/GameplayCameras/Private/Directors/SingleCameraDirector.cpp
@@ -2,7 +2,7 @@
 
 #include "Directors/SingleCameraDirector.h"
 
-#include "Build/CameraBuildLog.h"
+#include "Build/CameraBuildContext.h"
 #include "Core/CameraAsset.h"
 #include "Logging/TokenizedMessage.h"
 
@@ -42,11 +42,11 @@ FCameraDirectorEvaluatorPtr USingleCameraDirector::OnBuildEvaluator(FCameraDirec
 	return Builder.BuildEvaluator<FSingleCameraDirectorEvaluator>();
 }
 
-void USingleCameraDirector::OnBuildCameraDirector(UE::Cameras::FCameraBuildLog& BuildLog)
+void USingleCameraDirector::OnBuildCameraDirector(UE::Cameras::FCameraBuildContext& BuildContext)
 {
 	if (!CameraRig)
 	{
-		BuildLog.AddMessage(EMessageSeverity::Error, this, LOCTEXT("MissingCameraRig", "No camera rig is set."));
+		BuildContext.BuildLog.AddMessage(EMessageSeverity::Error, this, LOCTEXT("MissingCameraRig", "No camera rig is set."));
 	}
 }
 
diff --git a/Engine/Plugins/Cameras/GameplayCameras/Source/GameplayCameras/Private/Directors/StateTreeCameraDirector.cpp b/Engine/Plugins/Cameras/GameplayCameras/Source/GameplayCameras/Private/Directors/StateTreeCameraDirector.cpp
--- a/Engine/Plugins/Cameras/GameplayCameras/Source/GameplayCameras/Private/Directors/StateTreeCameraDirector.cpp
+++ b/Engine/Plugins/Cameras/GameplayCameras/Source/GameplayCameras/Private/Directors/StateTreeCameraDirector.cpp
@@ -2,7 +2,7 @@
 
 #include "Directors/StateTreeCameraDirector.h"
 
-#include "Build/CameraBuildLog.h"
+#include "Build/CameraBuildContext.h"
 #include "Core/CameraAsset.h"
 #include "Core/CameraEvaluationContext.h"
 #include "Core/CameraRigAsset.h"
@@ -227,14 +227,14 @@ FCameraDirectorEvaluatorPtr UStateTreeCameraDirector::OnBuildEvaluator(FCameraDi
 	return Builder.BuildEvaluator<FStateTreeCameraDirectorEvaluator>();
 }
 
-void UStateTreeCameraDirector::OnBuildCameraDirector(UE::Cameras::FCameraBuildLog& BuildLog)
+void UStateTreeCameraDirector::OnBuildCameraDirector(UE::Cameras::FCameraBuildContext& BuildContext)
 {
 	using namespace UE::Cameras;
 
 	// Check that a state tree was specified.
 	if (!StateTreeReference.IsValid())
 	{
-		BuildLog.AddMessage(EMessageSeverity::Error, LOCTEXT("MissingStateTree", "No state tree reference is set."));
+		BuildContext.BuildLog.AddMessage(EMessageSeverity::Error, LOCTEXT("MissingStateTree", "No state tree reference is set."));
 		return;
 	}
 }
diff --git a/Engine/Plugins/Cameras/GameplayCameras/Source/GameplayCameras/Private/GameFramework/GameplayCameraRigComponent.cpp b/Engine/Plugins/Cameras/GameplayCameras/Source/GameplayCameras/Private/GameFramework/GameplayCameraRigComponent.cpp
--- a/Engine/Plugins/Cameras/GameplayCameras/Source/GameplayCameras/Private/GameFramework/GameplayCameraRigComponent.cpp
+++ b/Engine/Plugins/Cameras/GameplayCameras/Source/GameplayCameras/Private/GameFramework/GameplayCameraRigComponent.cpp
@@ -3,7 +3,7 @@
 #include "GameFramework/GameplayCameraRigComponent.h"
 
 #include "Build/CameraAssetBuilder.h"
-#include "Build/CameraBuildLog.h"
+#include "Build/CameraBuildContext.h"
 #include "Core/CameraAsset.h"
 #include "Core/CameraRigAsset.h"
 #include "Directors/SingleCameraDirector.h"
@@ -101,7 +101,8 @@ void UGameplayCameraRigComponent::BuildGeneratedCamera()
 
 		FCameraBuildLog BuildLog;
 		BuildLog.SetForwardMessagesToLogging(true);
-		FCameraAssetBuilder Builder(BuildLog);
+		FCameraBuildContext BuildContext(BuildLog);
+		FCameraAssetBuilder Builder(BuildContext);
 		Builder.BuildCamera(GeneratedCameraAsset, bBuildReferencedAssets);
 	}
 }
diff --git a/Engine/Plugins/Cameras/GameplayCameras/Source/GameplayCameras/Private/Nodes/Common/CameraRigCameraNode.cpp b/Engine/Plugins/Cameras/GameplayCameras/Source/GameplayCameras/Private/Nodes/Common/CameraRigCameraNode.cpp
--- a/Engine/Plugins/Cameras/GameplayCameras/Source/GameplayCameras/Private/Nodes/Common/CameraRigCameraNode.cpp
+++ b/Engine/Plugins/Cameras/GameplayCameras/Source/GameplayCameras/Private/Nodes/Common/CameraRigCameraNode.cpp
@@ -44,8 +44,14 @@ void FCameraRigCameraNodeEvaluator::OnBuild(const FCameraNodeEvaluatorBuildParam
 
 void FCameraRigCameraNodeEvaluator::OnInitialize(const FCameraNodeEvaluatorInitializeParams& Params, FCameraNodeEvaluationResult& OutResult)
 {
-	// Apply overrides and defaults right away.
 	const UCameraRigCameraNode* PrefabNode = GetCameraNodeAs<UCameraRigCameraNode>();
+
+	// In editor and other uncooked builds, we only have the allocation info for the variables and data of our own 
+	// camera rig.  So we need to make room for the referenced camera rig's stuff.
+	// With cooked data we would have all that info already.
+	PrefabNode->CameraRigReference.EnsureAllocationInfo(OutResult);
+	
+	// Apply overrides and defaults right away.
 	PrefabNode->CameraRigReference.ApplyParameterOverridesAndDefaults(OutResult);
 }
 
@@ -66,16 +72,17 @@ void FCameraRigCameraNodeEvaluator::OnRun(const FCameraNodeEvaluationParams& Par
 
 }  // namespace UE::Cameras
 
-void UCameraRigCameraNode::OnPreBuild(FCameraBuildLog& BuildLog)
+void UCameraRigCameraNode::OnPreBuild(FCameraBuildContext& BuildContext)
 {
 	// Build the inner camera rig. Silently skip it if it's not set or invalid... but we will
 	// report an error in OnBuild about it.
 	if (UCameraRigAsset* CameraRig = CameraRigReference.GetCameraRig())
 	{
+		const bool bNeedsBuild = CameraRig->BuildStatus == ECameraBuildStatus::Dirty;
 		const UCameraRigAsset* OuterCameraRig = GetTypedOuter<UCameraRigAsset>();
-		if (OuterCameraRig != CameraRig)
+		if (OuterCameraRig != CameraRig && bNeedsBuild)
 		{
-			CameraRig->BuildCameraRig(BuildLog);
+			CameraRig->BuildCameraRig(BuildContext);
 		}
 	}
 
@@ -104,9 +111,12 @@ void UCameraRigCameraNode::OnBuild(FCameraObjectBuildContext& BuildContext)
 	const UCameraRigAsset* OuterCameraRig = GetTypedOuter<UCameraRigAsset>();
 	if (OuterCameraRig != CameraRig)
 	{
-		// Whatever allocations our inner camera rig needs for its evaluators and
-		// their camera variables, we add that to our camera rig's allocation info.
-		BuildContext.AllocationInfo.Append(CameraRig->AllocationInfo);
+		if (BuildContext.IsCooking())
+		{
+			// Whatever allocations our inner camera rig needs for its evaluators and
+			// their camera variables, we add that to our camera rig's allocation info.
+			BuildContext.AllocationInfo.Append(CameraRig->AllocationInfo);
+		}
 	}
 	else
 	{
diff --git a/Engine/Plugins/Cameras/GameplayCameras/Source/GameplayCameras/Private/Nodes/Shakes/CameraShakeCameraNode.cpp b/Engine/Plugins/Cameras/GameplayCameras/Source/GameplayCameras/Private/Nodes/Shakes/CameraShakeCameraNode.cpp
--- a/Engine/Plugins/Cameras/GameplayCameras/Source/GameplayCameras/Private/Nodes/Shakes/CameraShakeCameraNode.cpp
+++ b/Engine/Plugins/Cameras/GameplayCameras/Source/GameplayCameras/Private/Nodes/Shakes/CameraShakeCameraNode.cpp
@@ -72,8 +72,14 @@ void FCameraShakeCameraNodeEvaluator::OnBuild(const FCameraNodeEvaluatorBuildPar
 
 void FCameraShakeCameraNodeEvaluator::OnInitialize(const FCameraNodeEvaluatorInitializeParams& Params, FCameraNodeEvaluationResult& OutResult)
 {
-	// Apply overrides and defaults right away.
 	const UCameraShakeCameraNode* PrefabNode = GetCameraNodeAs<UCameraShakeCameraNode>();
+
+	// In editor and other uncooked builds, we don't have the allocation info for the variables and data of the
+	// reference shake because we try to keep assets standalone. So we need to make room for it.
+	// With cooked data we would have all that info already.
+	PrefabNode->CameraShakeReference.EnsureAllocationInfo(OutResult);
+
+	// Apply overrides and defaults right away.
 	PrefabNode->CameraShakeReference.ApplyParameterOverridesAndDefaults(OutResult);
 
 	// If evaluating the shake later in the visual layer, acquire the shake service we will use to
@@ -120,13 +126,17 @@ void FCameraShakeCameraNodeEvaluator::OnRun(const FCameraNodeEvaluationParams& P
 
 }  // namespace UE::Cameras
 
-void UCameraShakeCameraNode::OnPreBuild(FCameraBuildLog& BuildLog)
+void UCameraShakeCameraNode::OnPreBuild(FCameraBuildContext& BuildContext)
 {
 	// Build the inner camera shake. Silently skip it if it's not set... but we will
 	// report an error in OnBuild about it.
 	if (UCameraShakeAsset* CameraShake = CameraShakeReference.GetCameraShake())
 	{
-		CameraShake->BuildCameraShake(BuildLog);
+		const bool bNeedsBuild = CameraShake->BuildStatus == ECameraBuildStatus::Dirty;
+		if (bNeedsBuild)
+		{
+			CameraShake->BuildCameraShake(BuildContext);
+		}
 	}
 
 	// Make sure the property bag of the camera shake reference is up to date.
@@ -151,19 +161,22 @@ void UCameraShakeCameraNode::OnBuild(FCameraObjectBuildContext& BuildContext)
 		Modify();
 	}
 
-	// Whatever allocations our inner camera shake needs for its evaluators and
-	// their camera variables, we add that to our camera shake's allocation info.
-	// If we're going to be running the shake in a deferred way, however, we need
-	// to make the variables public so that they propagate to it.
-	FCameraObjectAllocationInfo CameraShakeAllocationInfo(CameraShake->AllocationInfo);
-	if (EvaluationMode == ECameraShakeEvaluationMode::VisualLayer)
+	if (BuildContext.IsCooking())
 	{
-		for (FCameraVariableDefinition& VariableDefinition : CameraShakeAllocationInfo.VariableTableInfo.VariableDefinitions)
+		// Whatever allocations our inner camera shake needs for its evaluators and
+		// their camera variables, we add that to our camera shake's allocation info.
+		// If we're going to be running the shake in a deferred way, however, we need
+		// to make the variables public so that they propagate to it.
+		FCameraObjectAllocationInfo CameraShakeAllocationInfo(CameraShake->AllocationInfo);
+		if (EvaluationMode == ECameraShakeEvaluationMode::VisualLayer)
 		{
-			VariableDefinition.bIsPrivate = false;
+			for (FCameraVariableDefinition& VariableDefinition : CameraShakeAllocationInfo.VariableTableInfo.VariableDefinitions)
+			{
+				VariableDefinition.bIsPrivate = false;
+			}
 		}
+		BuildContext.AllocationInfo.Append(CameraShakeAllocationInfo);
 	}
-	BuildContext.AllocationInfo.Append(CameraShakeAllocationInfo);
 }
 
 void UCameraShakeCameraNode::GetCustomCameraNodeParameters(FCameraNodeParameterInfos& OutParameterInfos)
diff --git a/Engine/Plugins/Cameras/GameplayCameras/Source/GameplayCameras/Private/Nodes/Utility/BlueprintCameraNode.cpp b/Engine/Plugins/Cameras/GameplayCameras/Source/GameplayCameras/Private/Nodes/Utility/BlueprintCameraNode.cpp
--- a/Engine/Plugins/Cameras/GameplayCameras/Source/GameplayCameras/Private/Nodes/Utility/BlueprintCameraNode.cpp
+++ b/Engine/Plugins/Cameras/GameplayCameras/Source/GameplayCameras/Private/Nodes/Utility/BlueprintCameraNode.cpp
@@ -595,7 +595,7 @@ void UBlueprintCameraNode::RebuildOverrides()
 	}
 }
 
-void UBlueprintCameraNode::OnPreBuild(FCameraBuildLog& BuildLog)
+void UBlueprintCameraNode::OnPreBuild(FCameraBuildContext& BuildContext)
 {
 	RebuildOverrides();
 }
diff --git a/Engine/Plugins/Cameras/GameplayCameras/Source/GameplayCameras/Public/Build/CameraAssetBuilder.h b/Engine/Plugins/Cameras/GameplayCameras/Source/GameplayCameras/Public/Build/CameraAssetBuilder.h
--- a/Engine/Plugins/Cameras/GameplayCameras/Source/GameplayCameras/Public/Build/CameraAssetBuilder.h
+++ b/Engine/Plugins/Cameras/GameplayCameras/Source/GameplayCameras/Public/Build/CameraAssetBuilder.h
@@ -2,7 +2,7 @@
 
 #pragma once
 
-#include "Build/CameraBuildLog.h"
+#include "Build/CameraBuildContext.h"
 #include "CoreTypes.h"
 #include "GameplayCameras.h"
 
@@ -23,24 +23,36 @@ class FCameraAssetBuilder
 public:
 
 	/** Creates a new camera builder. */
-	UE_API FCameraAssetBuilder(FCameraBuildLog& InBuildLog);
+	UE_API FCameraAssetBuilder(FCameraBuildContext& InBuildContext);
 
 	/** Builds the given camera. */
 	UE_API void BuildCamera(UCameraAsset* InCameraAsset, bool bBuildReferencedAssets = true);
 
+	/** Gets the last built referenced assets. */
+	TConstArrayView<UObject*> GetBuiltReferencedAssets() const { return BuiltReferencedAssets; }
+
+public:
+
+	// Internal API.
+
+	UE_API static bool BuildDefaultInterfaceIfNeeded(UCameraAsset* InCameraAsset);
+
 private:
 
 	void BuildCameraImpl(bool bBuildReferencedAssets);
 
-	void CopyDefaultParameterValues(const UCameraRigAsset* InCameraRig, FInstancedPropertyBag& DefaultParameters);
+	void BuildParameters();
+	void CopyDefaultParameterValue(const UCameraRigAsset* InCameraRig, const FGuid& PropertyID, FInstancedPropertyBag& DefaultParameters);
 
 	void UpdateBuildStatus();
 
 private:
 
-	FCameraBuildLog& BuildLog;
+	FCameraBuildContext BuildContext;
 
 	UCameraAsset* CameraAsset = nullptr;
+
+	TArray<UObject*> BuiltReferencedAssets;
 };
 
 }  // namespace UE::Cameras
diff --git a/Engine/Plugins/Cameras/GameplayCameras/Source/GameplayCameras/Public/Build/CameraBuildContext.h b/Engine/Plugins/Cameras/GameplayCameras/Source/GameplayCameras/Public/Build/CameraBuildContext.h
new file mode 100644
--- /dev/null
+++ b/Engine/Plugins/Cameras/GameplayCameras/Source/GameplayCameras/Public/Build/CameraBuildContext.h
@@ -0,0 +1,49 @@
+// Copyright Epic Games, Inc. All Rights Reserved.
+
+#pragma once
+
+#include "Build/CameraBuildLog.h"
+
+namespace UE::Cameras
+{
+
+enum ECameraBuildReason
+{
+	/** The reason for this build is undefined. */
+	Undefined,
+	/** A user clicked on the Build button, or similar. */
+	UserAction,
+	/** The build is happening during cooking. */
+	Cooking,
+	/** The build is happening for starting PIE. */
+	StartingPIE
+};
+
+/**
+ * Camera build context.
+ */
+struct FCameraBuildContext
+{
+	FCameraBuildContext(FCameraBuildLog& InBuildLog, ECameraBuildReason InBuildReason = ECameraBuildReason::Undefined)
+		: BuildLog(InBuildLog)
+		, BuildReason(InBuildReason)
+	{}
+
+	/** The build log for emitting messages. */
+	FCameraBuildLog& BuildLog;
+
+	/** The reason for this build. */
+	ECameraBuildReason GetBuildReason() const { return BuildReason; }
+
+public:
+
+	/** Whether the build reason is for cooking. */
+	bool IsCooking() const { return BuildReason == ECameraBuildReason::Cooking; }
+
+private:
+
+	ECameraBuildReason BuildReason = ECameraBuildReason::Undefined;
+};
+
+}  // namespace UE::Cameras
+
diff --git a/Engine/Plugins/Cameras/GameplayCameras/Source/GameplayCameras/Public/Build/CameraNodeHierarchyBuilder.h b/Engine/Plugins/Cameras/GameplayCameras/Source/GameplayCameras/Public/Build/CameraNodeHierarchyBuilder.h
--- a/Engine/Plugins/Cameras/GameplayCameras/Source/GameplayCameras/Public/Build/CameraNodeHierarchyBuilder.h
+++ b/Engine/Plugins/Cameras/GameplayCameras/Source/GameplayCameras/Public/Build/CameraNodeHierarchyBuilder.h
@@ -2,7 +2,7 @@
 
 #pragma once
 
-#include "Build/CameraBuildLog.h"
+#include "Build/CameraBuildContext.h"
 #include "Core/CameraNodeHierarchy.h"
 
 #define UE_API GAMEPLAYCAMERAS_API
@@ -22,7 +22,7 @@ class FCameraNodeHierarchyBuilder
 public:
 
 	/** Creates a new camera node hierarchy builder. */
-	UE_API FCameraNodeHierarchyBuilder(FCameraBuildLog& InBuildLog, UBaseCameraObject* InCameraObject);
+	UE_API FCameraNodeHierarchyBuilder(FCameraBuildContext& InBuildContext, UBaseCameraObject* InCameraObject);
 
 	/** Gets the camera node hierarchy. */
 	const FCameraNodeHierarchy& GetHierarchy() const { return CameraNodeHierarchy; }
@@ -40,7 +40,7 @@ private:
 
 private:
 
-	FCameraBuildLog& BuildLog;
+	FCameraBuildContext BuildContext;
 	UBaseCameraObject* CameraObject = nullptr;
 	FCameraNodeHierarchy CameraNodeHierarchy;
 };
diff --git a/Engine/Plugins/Cameras/GameplayCameras/Source/GameplayCameras/Public/Build/CameraObjectBuildContext.h b/Engine/Plugins/Cameras/GameplayCameras/Source/GameplayCameras/Public/Build/CameraObjectBuildContext.h
--- a/Engine/Plugins/Cameras/GameplayCameras/Source/GameplayCameras/Public/Build/CameraObjectBuildContext.h
+++ b/Engine/Plugins/Cameras/GameplayCameras/Source/GameplayCameras/Public/Build/CameraObjectBuildContext.h
@@ -2,25 +2,21 @@
 
 #pragma once
 
+#include "Build/CameraBuildContext.h"
 #include "Core/BaseCameraObject.h"
 
 namespace UE::Cameras
 {
 
-class FCameraBuildLog;
-
 /**
  * Camera object build context.
  */
-struct FCameraObjectBuildContext
+struct FCameraObjectBuildContext : public FCameraBuildContext
 {
-	FCameraObjectBuildContext(FCameraBuildLog& InBuildLog)
-		: BuildLog(InBuildLog)
+	FCameraObjectBuildContext(const FCameraBuildContext& InParentContext)
+		: FCameraBuildContext(InParentContext)
 	{}
 
-	/** The build log for emitting messages. */
-	FCameraBuildLog& BuildLog;
-
 	/** The allocation information for the camera rig. */
 	FCameraObjectAllocationInfo AllocationInfo;
 };
diff --git a/Engine/Plugins/Cameras/GameplayCameras/Source/GameplayCameras/Public/Build/CameraObjectInterfaceBuilder.h b/Engine/Plugins/Cameras/GameplayCameras/Source/GameplayCameras/Public/Build/CameraObjectInterfaceBuilder.h
--- a/Engine/Plugins/Cameras/GameplayCameras/Source/GameplayCameras/Public/Build/CameraObjectInterfaceBuilder.h
+++ b/Engine/Plugins/Cameras/GameplayCameras/Source/GameplayCameras/Public/Build/CameraObjectInterfaceBuilder.h
@@ -2,7 +2,7 @@
 
 #pragma once
 
-#include "Build/CameraBuildLog.h"
+#include "Build/CameraBuildContext.h"
 #include "Containers/ArrayView.h"
 #include "Core/CameraContextDataTableFwd.h"
 #include "Core/CameraVariableTableFwd.h"
@@ -26,7 +26,7 @@ class FCameraObjectInterfaceBuilder
 {
 public:
 
-	UE_API FCameraObjectInterfaceBuilder(FCameraBuildLog& InBuildLog);
+	UE_API FCameraObjectInterfaceBuilder(FCameraBuildContext& InBuildContext);
 
 	UE_API void BuildInterface(UBaseCameraObject* InCameraObject, const FCameraNodeHierarchy& InHierarchy, bool bCollectStrayNodes);
 	UE_API void BuildInterface(UBaseCameraObject* InCameraObject, TArrayView<UCameraNode*> InCameraObjectNodes);
@@ -50,7 +50,7 @@ private:
 
 private:
 
-	FCameraBuildLog& BuildLog;
+	FCameraBuildContext BuildContext;
 
 	UBaseCameraObject* CameraObject = nullptr;
 	TArray<UCameraNode*> CameraObjectNodes;
diff --git a/Engine/Plugins/Cameras/GameplayCameras/Source/GameplayCameras/Public/Build/CameraRigAssetBuilder.h b/Engine/Plugins/Cameras/GameplayCameras/Source/GameplayCameras/Public/Build/CameraRigAssetBuilder.h
--- a/Engine/Plugins/Cameras/GameplayCameras/Source/GameplayCameras/Public/Build/CameraRigAssetBuilder.h
+++ b/Engine/Plugins/Cameras/GameplayCameras/Source/GameplayCameras/Public/Build/CameraRigAssetBuilder.h
@@ -2,7 +2,7 @@
 
 #pragma once
 
-#include "Build/CameraBuildLog.h"
+#include "Build/CameraBuildContext.h"
 #include "Core/CameraNodeHierarchy.h"
 
 #define UE_API GAMEPLAYCAMERAS_API
@@ -26,7 +26,7 @@ class FCameraRigAssetBuilder
 public:
 
 	/** Creates a new camera rig builder. */
-	UE_API FCameraRigAssetBuilder(FCameraBuildLog& InBuildLog);
+	UE_API FCameraRigAssetBuilder(FCameraBuildContext& InBuildContext);
 
 	/** Builds the given camera rig. */
 	UE_API void BuildCameraRig(UCameraRigAsset* InCameraRig);
@@ -39,7 +39,7 @@ private:
 
 private:
 
-	FCameraBuildLog& BuildLog;
+	FCameraBuildContext BuildContext;
 
 	UCameraRigAsset* CameraRig = nullptr;
 
diff --git a/Engine/Plugins/Cameras/GameplayCameras/Source/GameplayCameras/Public/Build/CameraShakeAssetBuilder.h b/Engine/Plugins/Cameras/GameplayCameras/Source/GameplayCameras/Public/Build/CameraShakeAssetBuilder.h
--- a/Engine/Plugins/Cameras/GameplayCameras/Source/GameplayCameras/Public/Build/CameraShakeAssetBuilder.h
+++ b/Engine/Plugins/Cameras/GameplayCameras/Source/GameplayCameras/Public/Build/CameraShakeAssetBuilder.h
@@ -2,7 +2,7 @@
 
 #pragma once
 
-#include "Build/CameraBuildLog.h"
+#include "Build/CameraBuildContext.h"
 #include "Core/CameraNodeHierarchy.h"
 
 #define UE_API GAMEPLAYCAMERAS_API
@@ -16,7 +16,7 @@ class FCameraShakeAssetBuilder
 {
 public:
 
-	UE_API FCameraShakeAssetBuilder(FCameraBuildLog& InBuildLog);
+	UE_API FCameraShakeAssetBuilder(FCameraBuildContext& InBuildContext);
 
 	UE_API void BuildCameraShake(UCameraShakeAsset* InCameraShake);
 
@@ -28,7 +28,7 @@ private:
 
 private:
 
-	FCameraBuildLog& BuildLog;
+	FCameraBuildContext BuildContext;
 
 	UCameraShakeAsset* CameraShake = nullptr;
 
diff --git a/Engine/Plugins/Cameras/GameplayCameras/Source/GameplayCameras/Public/Core/BaseCameraObject.h b/Engine/Plugins/Cameras/GameplayCameras/Source/GameplayCameras/Public/Core/BaseCameraObject.h
--- a/Engine/Plugins/Cameras/GameplayCameras/Source/GameplayCameras/Public/Core/BaseCameraObject.h
+++ b/Engine/Plugins/Cameras/GameplayCameras/Source/GameplayCameras/Public/Core/BaseCameraObject.h
@@ -165,6 +165,14 @@ public:
 	/** Gets the definitions of parameters exposed on this camera rig. */
 	TConstArrayView<FCameraObjectInterfaceParameterDefinition> GetParameterDefinitions() const { return ParameterDefinitions; }
 
+public:
+
+	/** Finds a parameter definition by name. */
+	UE_API bool FindParameterDefinitionByName(const FName ParameterName, FCameraObjectInterfaceParameterDefinition& OutParameterDefinition) const;
+
+	/** Finds a parameter definition by ID. */
+	UE_API bool FindParameterDefinitionByGuid(const FGuid& ParameterGuid, FCameraObjectInterfaceParameterDefinition& OutParameterDefinition) const;
+
 public:
 
 	/** Get the root node of this camera object. */
diff --git a/Engine/Plugins/Cameras/GameplayCameras/Source/GameplayCameras/Public/Core/CameraAsset.h b/Engine/Plugins/Cameras/GameplayCameras/Source/GameplayCameras/Public/Core/CameraAsset.h
--- a/Engine/Plugins/Cameras/GameplayCameras/Source/GameplayCameras/Public/Core/CameraAsset.h
+++ b/Engine/Plugins/Cameras/GameplayCameras/Source/GameplayCameras/Public/Core/CameraAsset.h
@@ -23,8 +23,7 @@ class UCameraRigAsset;
 namespace UE::Cameras
 {
 	class FCameraAssetBuilder;
-	class FCameraAssetParameterOverrideEvaluator;
-	class FCameraBuildLog;
+	struct FCameraBuildContext;
 
 	/**
 	 * Interface for listening to changes on a camera asset.
@@ -40,9 +39,70 @@ namespace UE::Cameras
 		virtual void OnEnterTransitionsChanged(UCameraAsset* InCameraAsset, const TCameraArrayChangedEvent<UCameraRigTransition*>& Event) {}
 		/* Changed when the exit transitions have been changed. */
 		virtual void OnExitTransitionsChanged(UCameraAsset* InCameraAsset, const TCameraArrayChangedEvent<UCameraRigTransition*>& Event) {}
+		/** Called when the camera asset's interface has changed. */
+		virtual void OnCameraAssetInterfaceChanged() {}
 	};
 }
 
+/**
+ * An interface parameter on a camera asset, which exposes one of its camera rigs' own interface parameters.
+ */
+UCLASS()
+class UCameraAssetInterfaceParameter : public UObject
+{
+	GENERATED_BODY()
+
+public:
+
+	/** The exposed name for this parameter. */
+	UPROPERTY(EditAnywhere, Category="Camera")
+	FString InterfaceParameterName;
+
+	/** The camera rig that owns the source parameter. */
+	UPROPERTY(EditAnywhere, Category="Camera")
+	TSoftObjectPtr<UCameraRigAsset> SourceCameraRig;
+
+	/** The name of the parameter on the source camera rig. */
+	UPROPERTY(EditAnywhere, Category="Camera")
+	FName SourceParameterName;
+
+public:
+
+	/** Gets this parameter's unique ID. */
+	const FGuid& GetGuid() const { return Guid; }
+
+	/** Gets the parameter definition. */
+	bool GetParameterDefinition(FCameraObjectInterfaceParameterDefinition& OutParameterDefinition) const;
+
+protected:
+
+	/** The Guid of this parameter. */
+	UPROPERTY()
+	FGuid Guid;
+
+protected:
+
+	// UObject interface.
+	virtual void PostLoad() override;
+	virtual void PostInitProperties() override;
+	virtual void PostDuplicate(EDuplicateMode::Type DuplicateMode) override;
+};
+
+/**
+ * A structure defining the public interface of a camera asset.
+ */
+USTRUCT()
+struct FCameraAssetInterface
+{
+	GENERATED_BODY()
+
+public:
+
+	/** The list of exposed parameters on the camera asset. */
+	UPROPERTY(Instanced)
+	TArray<TObjectPtr<UCameraAssetInterfaceParameter>> Parameters;
+};
+
 /**
  * Structure describing various allocations needed by a camera asset.
  */
@@ -115,9 +175,9 @@ public:
 
 	/**
 	 * Builds and validates this camera, including all its camera rigs.
-	 * Errors and warnings will go to the provided build log.
+	 * Errors and warnings will go to the provided build context.
 	 */
-	GAMEPLAYCAMERAS_API void BuildCamera(UE::Cameras::FCameraBuildLog& InBuildLog);
+	GAMEPLAYCAMERAS_API void BuildCamera(UE::Cameras::FCameraBuildContext& InBuildContext);
 
 public:
 
@@ -165,6 +225,10 @@ public:
 	/** Event handlers to be notified of data changes. */
 	UE::Cameras::TCameraEventHandlerContainer<UE::Cameras::ICameraAssetEventHandler> EventHandlers;
 
+	/** The public interface of this camera asset. */
+	UPROPERTY()
+	FCameraAssetInterface Interface;
+
 private:
 
 	/** The camera director to use in this camera. */
@@ -195,12 +259,6 @@ private:
 	UPROPERTY()
 	FCameraAssetAllocationInfo AllocationInfo;
 
-
-	// Deprecated.
-
-	UPROPERTY()
-	TArray<TObjectPtr<UCameraRigAsset>> CameraRigs_DEPRECATED;
-
 #if WITH_EDITORONLY_DATA
 
 	/** Position of the camera node in the shared transitions graph editor. */
@@ -219,8 +277,10 @@ private:
 	friend class UCameraSharedTransitionGraphSchema;
 
 #endif  // WITH_EDITORONLY_DATA
-	
+
+	// Flag for whether we need to upgrade the interface on the next build.
+	bool bNeedsDefaultInterface = false;
+
 	friend class UE::Cameras::FCameraAssetBuilder;
-	friend class UE::Cameras::FCameraAssetParameterOverrideEvaluator;
 };
 
diff --git a/Engine/Plugins/Cameras/GameplayCameras/Source/GameplayCameras/Public/Core/CameraContextDataTable.h b/Engine/Plugins/Cameras/GameplayCameras/Source/GameplayCameras/Public/Core/CameraContextDataTable.h
--- a/Engine/Plugins/Cameras/GameplayCameras/Source/GameplayCameras/Public/Core/CameraContextDataTable.h
+++ b/Engine/Plugins/Cameras/GameplayCameras/Source/GameplayCameras/Public/Core/CameraContextDataTable.h
@@ -59,6 +59,12 @@ public:
 	/** Initializes the context data table so that it fits the provided allocation info. */
 	void Initialize(const FCameraContextDataTableAllocationInfo& AllocationInfo);
 
+	/**
+	 * Ensure that all the data referenced in the given allocation info exist in the table.
+	 * Returns false if some data had to be added.
+	 */
+	bool EnsureData(const FCameraContextDataTableAllocationInfo& AllocationInfo);
+
 	/** Adds a data entry to the table. */
 	void AddData(const FCameraContextDataDefinition& DataDefinition);
 
@@ -269,6 +275,8 @@ private:
 	static bool SetDataValue(ECameraContextDataType DataType, ECameraContextDataContainerType DataContainerType, const UObject* DataTypeObject, uint8* DestDataPtr, const uint8* SrcDataPtr);
 	static bool SetDataValue(ECameraContextDataType DataType, const UObject* DataTypeObject, uint8* DestDataPtr, const uint8* SrcDataPtr);
 
+	static void CreateEntry(const FCameraContextDataDefinition& DataDefinition, FEntry& OutEntry);
+
 	const FEntry* FindEntry(FCameraContextDataID InID) const;
 	FEntry* FindEntry(FCameraContextDataID InID);
 
@@ -296,6 +304,7 @@ private:
 
 	uint8* Memory = nullptr;
 	uint32 Capacity = 0;
+	uint32 Alignment = 0;
 	uint32 Used = 0;
 
 #if UE_GAMEPLAY_CAMERAS_DEBUG
diff --git a/Engine/Plugins/Cameras/GameplayCameras/Source/GameplayCameras/Public/Core/CameraContextDataTableAllocationInfo.h b/Engine/Plugins/Cameras/GameplayCameras/Source/GameplayCameras/Public/Core/CameraContextDataTableAllocationInfo.h
--- a/Engine/Plugins/Cameras/GameplayCameras/Source/GameplayCameras/Public/Core/CameraContextDataTableAllocationInfo.h
+++ b/Engine/Plugins/Cameras/GameplayCameras/Source/GameplayCameras/Public/Core/CameraContextDataTableAllocationInfo.h
@@ -63,6 +63,9 @@ struct FCameraContextDataTableAllocationInfo
 	/**Combines the given allocation info with this one. */
 	GAMEPLAYCAMERAS_API void Combine(const FCameraContextDataTableAllocationInfo& OtherInfo);
 
+	/** Returns whether the given data definitions already exist in our list. */
+	GAMEPLAYCAMERAS_API bool Contains(const FCameraContextDataTableAllocationInfo& OtherInfo) const;
+
 	bool operator==(const FCameraContextDataTableAllocationInfo& Other) const = default;
 };
 
diff --git a/Engine/Plugins/Cameras/GameplayCameras/Source/GameplayCameras/Public/Core/CameraDirector.h b/Engine/Plugins/Cameras/GameplayCameras/Source/GameplayCameras/Public/Core/CameraDirector.h
--- a/Engine/Plugins/Cameras/GameplayCameras/Source/GameplayCameras/Public/Core/CameraDirector.h
+++ b/Engine/Plugins/Cameras/GameplayCameras/Source/GameplayCameras/Public/Core/CameraDirector.h
@@ -15,7 +15,7 @@
 class UCameraRigAsset;
 class UCameraRigProxy;
 
-namespace UE::Cameras { class FCameraBuildLog; }
+namespace UE::Cameras { struct FCameraBuildContext; }
 
 #if WITH_EDITOR
 
@@ -56,7 +56,7 @@ public:
 	UE_API FCameraDirectorEvaluatorPtr BuildEvaluator(FCameraDirectorEvaluatorBuilder& Builder) const;
 
 	/** Builds and validates this camera director. */
-	UE_API void BuildCameraDirector(UE::Cameras::FCameraBuildLog& BuildLog);
+	UE_API void BuildCameraDirector(UE::Cameras::FCameraBuildContext& BuildContext);
 
 	/** Gets the list of camera rigs used by this camera director. */
 	UE_API void GatherRigUsageInfo(FCameraDirectorRigUsageInfo& UsageInfo) const;
@@ -80,7 +80,7 @@ protected:
 	virtual FCameraDirectorEvaluatorPtr OnBuildEvaluator(FCameraDirectorEvaluatorBuilder& Builder) const { return nullptr; }
 
 	/** Builds and validates this camera director. */
-	virtual void OnBuildCameraDirector(UE::Cameras::FCameraBuildLog& BuildLog) {}
+	virtual void OnBuildCameraDirector(UE::Cameras::FCameraBuildContext& BuildContext) {}
 
 	/** Gets the list of camera rigs used by this camera director. */
 	virtual void OnGatherRigUsageInfo(FCameraDirectorRigUsageInfo& UsageInfo) const {}
diff --git a/Engine/Plugins/Cameras/GameplayCameras/Source/GameplayCameras/Public/Core/CameraNode.h b/Engine/Plugins/Cameras/GameplayCameras/Source/GameplayCameras/Public/Core/CameraNode.h
--- a/Engine/Plugins/Cameras/GameplayCameras/Source/GameplayCameras/Public/Core/CameraNode.h
+++ b/Engine/Plugins/Cameras/GameplayCameras/Source/GameplayCameras/Public/Core/CameraNode.h
@@ -15,6 +15,7 @@
 namespace UE::Cameras
 {
 	class FCameraBuildLog;
+	struct FCameraBuildContext;
 	struct FCameraObjectBuildContext;
 }
 
@@ -44,6 +45,7 @@ class UCameraNode
 public:
 	
 	using FCameraBuildLog = UE::Cameras::FCameraBuildLog;
+	using FCameraBuildContext = UE::Cameras::FCameraBuildContext;
 	using FCameraObjectBuildContext = UE::Cameras::FCameraObjectBuildContext;
 	using FCameraNodeEvaluatorBuilder = UE::Cameras::FCameraNodeEvaluatorBuilder;
 
@@ -51,7 +53,7 @@ public:
 	UE_API FCameraNodeChildrenView GetChildren();
 
 	/** Optional build step executed at the beginning of the build process. */
-	UE_API void PreBuild(FCameraBuildLog& BuildLog);
+	UE_API void PreBuild(FCameraBuildContext& BuildContext);
 
 	/** Gets optional info about this node's required allocations at runtime. */
 	UE_API void Build(FCameraObjectBuildContext& BuildContext);
@@ -73,7 +75,7 @@ protected:
 	virtual FCameraNodeChildrenView OnGetChildren() { return FCameraNodeChildrenView(); }
 
 	/** Optional build step executed at the beginning of the build process. */
-	virtual void OnPreBuild(FCameraBuildLog& BuildLog) {}
+	virtual void OnPreBuild(FCameraBuildContext& BuildContext) {}
 
 	/** Gets optional info about this node's required allocations at runtime. */
 	virtual void OnBuild(FCameraObjectBuildContext& BuildContext) {}
diff --git a/Engine/Plugins/Cameras/GameplayCameras/Source/GameplayCameras/Public/Core/CameraRigAsset.h b/Engine/Plugins/Cameras/GameplayCameras/Source/GameplayCameras/Public/Core/CameraRigAsset.h
--- a/Engine/Plugins/Cameras/GameplayCameras/Source/GameplayCameras/Public/Core/CameraRigAsset.h
+++ b/Engine/Plugins/Cameras/GameplayCameras/Source/GameplayCameras/Public/Core/CameraRigAsset.h
@@ -22,8 +22,8 @@ class UCameraVariableAsset;
 
 namespace UE::Cameras
 {
-	class FCameraBuildLog;
 	class FCameraRigAssetBuilder;
+	struct FCameraBuildContext;
 
 	/**
 	 * Interface for listening to changes on a camera rig asset.
@@ -105,9 +105,9 @@ public:
 	GAMEPLAYCAMERAS_API void BuildCameraRig();
 
 	/**
-	 * Builds this camera rig, similar to BuildCameraRig() but using a given build log.
+	 * Builds this camera rig, similar to BuildCameraRig() but using a given build context.
 	 */
-	GAMEPLAYCAMERAS_API void BuildCameraRig(UE::Cameras::FCameraBuildLog& InBuildLog);
+	GAMEPLAYCAMERAS_API void BuildCameraRig(UE::Cameras::FCameraBuildContext& InBuildContext);
 
 public:
 
diff --git a/Engine/Plugins/Cameras/GameplayCameras/Source/GameplayCameras/Public/Core/CameraRigAssetReference.h b/Engine/Plugins/Cameras/GameplayCameras/Source/GameplayCameras/Public/Core/CameraRigAssetReference.h
--- a/Engine/Plugins/Cameras/GameplayCameras/Source/GameplayCameras/Public/Core/CameraRigAssetReference.h
+++ b/Engine/Plugins/Cameras/GameplayCameras/Source/GameplayCameras/Public/Core/CameraRigAssetReference.h
@@ -279,6 +279,11 @@ public:
 		}
 	}
 
+public:
+
+	/** Ensure the variable and data tables have the allocations for our referenced camera rig. */
+	void EnsureAllocationInfo(UE::Cameras::FCameraNodeEvaluationResult& OutResult) const;
+
 	/** Applies the parameter override values to the given variable table. */
 	void ApplyParameterOverrides(UE::Cameras::FCameraVariableTable& OutVariableTable, bool bDrivenOnly) const;
 
diff --git a/Engine/Plugins/Cameras/GameplayCameras/Source/GameplayCameras/Public/Core/CameraShakeAsset.h b/Engine/Plugins/Cameras/GameplayCameras/Source/GameplayCameras/Public/Core/CameraShakeAsset.h
--- a/Engine/Plugins/Cameras/GameplayCameras/Source/GameplayCameras/Public/Core/CameraShakeAsset.h
+++ b/Engine/Plugins/Cameras/GameplayCameras/Source/GameplayCameras/Public/Core/CameraShakeAsset.h
@@ -15,8 +15,8 @@ class UShakeCameraNode;
 
 namespace UE::Cameras
 {
-	class FCameraBuildLog;
 	class FCameraShakeAssetBuilder;
+	struct FCameraBuildContext;
 }
 
 UCLASS(MinimalAPI, BlueprintType)
@@ -59,7 +59,7 @@ public:
 	GAMEPLAYCAMERAS_API void BuildCameraShake();
 
 	/** Builds this camera shake asset. */
-	GAMEPLAYCAMERAS_API void BuildCameraShake(UE::Cameras::FCameraBuildLog& InBuildLog);
+	GAMEPLAYCAMERAS_API void BuildCameraShake(UE::Cameras::FCameraBuildContext& InBuildContext);
 
 public:
 
diff --git a/Engine/Plugins/Cameras/GameplayCameras/Source/GameplayCameras/Public/Core/CameraShakeAssetReference.h b/Engine/Plugins/Cameras/GameplayCameras/Source/GameplayCameras/Public/Core/CameraShakeAssetReference.h
--- a/Engine/Plugins/Cameras/GameplayCameras/Source/GameplayCameras/Public/Core/CameraShakeAssetReference.h
+++ b/Engine/Plugins/Cameras/GameplayCameras/Source/GameplayCameras/Public/Core/CameraShakeAssetReference.h
@@ -61,6 +61,11 @@ public:
 		}
 	}
 
+public:
+
+	/** Ensure the variable and data tables have the allocations for our referenced camera shake. */
+	void EnsureAllocationInfo(UE::Cameras::FCameraNodeEvaluationResult& OutResult) const;
+
 	/** Applies the parameter override values to the given variable table. */
 	void ApplyParameterOverrides(UE::Cameras::FCameraVariableTable& OutVariableTable, bool bDrivenOnly) const;
 
diff --git a/Engine/Plugins/Cameras/GameplayCameras/Source/GameplayCameras/Public/Core/CameraVariableTable.h b/Engine/Plugins/Cameras/GameplayCameras/Source/GameplayCameras/Public/Core/CameraVariableTable.h
--- a/Engine/Plugins/Cameras/GameplayCameras/Source/GameplayCameras/Public/Core/CameraVariableTable.h
+++ b/Engine/Plugins/Cameras/GameplayCameras/Source/GameplayCameras/Public/Core/CameraVariableTable.h
@@ -83,6 +83,12 @@ public:
 	/** Initializes the variable table so that it fits the provided allocation info. */
 	UE_API void Initialize(const FCameraVariableTableAllocationInfo& AllocationInfo);
 
+	/**
+	 * Ensure that all the variables referenced in the given allocation info exist in the table.
+	 * Returns false if some variables had to be added.
+	 */
+	UE_API bool EnsureVariables(const FCameraVariableTableAllocationInfo& AllocationInfo);
+
 	/** Adds a variable to the table.
 	 *
 	 * This may re-allocate the internal memory buffer. It's recommended to pre-compute
@@ -206,6 +212,8 @@ private:
 
 	static UE_API bool GetVariableTypeAllocationInfo(ECameraVariableType VariableType, const UScriptStruct* StructType, uint32& OutSizeOf, uint32& OutAlignOf);
 
+	static void CreateEntry(const FCameraVariableDefinition& VariableDefinition, FEntry& OutEntry);
+
 	template<typename ValueType>
 	static bool CheckVariableType(ECameraVariableType InType)
 	{
@@ -220,6 +228,8 @@ private:
 	UE_API void InternalOverride(const FCameraVariableTable& OtherTable, ECameraVariableTableFilter Filter, const FCameraVariableTableFlags* InMask, bool bInvertMask, FCameraVariableTableFlags* OutMask);
 	UE_API void InternalLerp(const FCameraVariableTable& ToTable, ECameraVariableTableFilter Filter, float Factor, const FCameraVariableTableFlags* InMask, bool bInvertMask, FCameraVariableTableFlags* OutMask);
 
+	void InitializeEntries(int32 StartIndex = 0, int32 EndIndex = -1);
+
 private:
 
 	enum class EEntryFlags : uint8
@@ -250,6 +260,7 @@ private:
 
 	uint8* Memory = nullptr;
 	uint32 Capacity = 0;
+	uint32 Alignment = 0;
 	uint32 Used = 0;
 
 	static UE_API TArray<FBlendableStructInfo> CachedBlendableStructs;
diff --git a/Engine/Plugins/Cameras/GameplayCameras/Source/GameplayCameras/Public/Core/CameraVariableTableAllocationInfo.h b/Engine/Plugins/Cameras/GameplayCameras/Source/GameplayCameras/Public/Core/CameraVariableTableAllocationInfo.h
--- a/Engine/Plugins/Cameras/GameplayCameras/Source/GameplayCameras/Public/Core/CameraVariableTableAllocationInfo.h
+++ b/Engine/Plugins/Cameras/GameplayCameras/Source/GameplayCameras/Public/Core/CameraVariableTableAllocationInfo.h
@@ -109,6 +109,9 @@ struct FCameraVariableTableAllocationInfo
 	/**Combines the given allocation info with this one. */
 	GAMEPLAYCAMERAS_API void Combine(const FCameraVariableTableAllocationInfo& OtherInfo);
 
+	/** Returns whether the given variable definitions already exist in our list. */
+	GAMEPLAYCAMERAS_API bool Contains(const FCameraVariableTableAllocationInfo& OtherInfo) const;
+
 	bool operator==(const FCameraVariableTableAllocationInfo& Other) const = default;
 };
 
diff --git a/Engine/Plugins/Cameras/GameplayCameras/Source/GameplayCameras/Public/Directors/BlueprintCameraDirector.h b/Engine/Plugins/Cameras/GameplayCameras/Source/GameplayCameras/Public/Directors/BlueprintCameraDirector.h
--- a/Engine/Plugins/Cameras/GameplayCameras/Source/GameplayCameras/Public/Directors/BlueprintCameraDirector.h
+++ b/Engine/Plugins/Cameras/GameplayCameras/Source/GameplayCameras/Public/Directors/BlueprintCameraDirector.h
@@ -235,7 +235,7 @@ protected:
 
 	// UCameraDirector interface.
 	virtual FCameraDirectorEvaluatorPtr OnBuildEvaluator(FCameraDirectorEvaluatorBuilder& Builder) const override;
-	virtual void OnBuildCameraDirector(UE::Cameras::FCameraBuildLog& BuildLog) override;
+	virtual void OnBuildCameraDirector(UE::Cameras::FCameraBuildContext& BuildContext) override;
 	virtual void OnGatherRigUsageInfo(FCameraDirectorRigUsageInfo& UsageInfo) const override;
 	virtual void OnExtendAssetRegistryTags(FAssetRegistryTagsContext Context) const override;
 };
diff --git a/Engine/Plugins/Cameras/GameplayCameras/Source/GameplayCameras/Public/Directors/SingleCameraDirector.h b/Engine/Plugins/Cameras/GameplayCameras/Source/GameplayCameras/Public/Directors/SingleCameraDirector.h
--- a/Engine/Plugins/Cameras/GameplayCameras/Source/GameplayCameras/Public/Directors/SingleCameraDirector.h
+++ b/Engine/Plugins/Cameras/GameplayCameras/Source/GameplayCameras/Public/Directors/SingleCameraDirector.h
@@ -22,7 +22,7 @@ protected:
 
 	// UCameraDirector interface.
 	virtual FCameraDirectorEvaluatorPtr OnBuildEvaluator(FCameraDirectorEvaluatorBuilder& Builder) const override;
-	virtual void OnBuildCameraDirector(UE::Cameras::FCameraBuildLog& BuildLog) override;
+	virtual void OnBuildCameraDirector(UE::Cameras::FCameraBuildContext& BuildContext) override;
 	virtual void OnGatherRigUsageInfo(FCameraDirectorRigUsageInfo& UsageInfo) const override;
 
 public:
diff --git a/Engine/Plugins/Cameras/GameplayCameras/Source/GameplayCameras/Public/Directors/StateTreeCameraDirector.h b/Engine/Plugins/Cameras/GameplayCameras/Source/GameplayCameras/Public/Directors/StateTreeCameraDirector.h
--- a/Engine/Plugins/Cameras/GameplayCameras/Source/GameplayCameras/Public/Directors/StateTreeCameraDirector.h
+++ b/Engine/Plugins/Cameras/GameplayCameras/Source/GameplayCameras/Public/Directors/StateTreeCameraDirector.h
@@ -24,7 +24,7 @@ protected:
 
 	// UCameraDirector interface.
 	virtual FCameraDirectorEvaluatorPtr OnBuildEvaluator(FCameraDirectorEvaluatorBuilder& Builder) const override;
-	virtual void OnBuildCameraDirector(UE::Cameras::FCameraBuildLog& BuildLog) override;
+	virtual void OnBuildCameraDirector(UE::Cameras::FCameraBuildContext& BuildContext) override;
 	virtual void OnGatherRigUsageInfo(FCameraDirectorRigUsageInfo& UsageInfo) const override;
 	virtual void OnExtendAssetRegistryTags(FAssetRegistryTagsContext Context) const override;
 
diff --git a/Engine/Plugins/Cameras/GameplayCameras/Source/GameplayCameras/Public/Nodes/Common/CameraRigCameraNode.h b/Engine/Plugins/Cameras/GameplayCameras/Source/GameplayCameras/Public/Nodes/Common/CameraRigCameraNode.h
--- a/Engine/Plugins/Cameras/GameplayCameras/Source/GameplayCameras/Public/Nodes/Common/CameraRigCameraNode.h
+++ b/Engine/Plugins/Cameras/GameplayCameras/Source/GameplayCameras/Public/Nodes/Common/CameraRigCameraNode.h
@@ -24,7 +24,7 @@ class UCameraRigCameraNode
 protected:
 
 	// UCameraNode interface.
-	virtual void OnPreBuild(FCameraBuildLog& BuildLog) override;
+	virtual void OnPreBuild(FCameraBuildContext& BuildContext) override;
 	virtual void OnBuild(FCameraObjectBuildContext& BuildContext) override;
 	virtual FCameraNodeEvaluatorPtr OnBuildEvaluator(FCameraNodeEvaluatorBuilder& Builder) const override;
 
diff --git a/Engine/Plugins/Cameras/GameplayCameras/Source/GameplayCameras/Public/Nodes/Shakes/CameraShakeCameraNode.h b/Engine/Plugins/Cameras/GameplayCameras/Source/GameplayCameras/Public/Nodes/Shakes/CameraShakeCameraNode.h
--- a/Engine/Plugins/Cameras/GameplayCameras/Source/GameplayCameras/Public/Nodes/Shakes/CameraShakeCameraNode.h
+++ b/Engine/Plugins/Cameras/GameplayCameras/Source/GameplayCameras/Public/Nodes/Shakes/CameraShakeCameraNode.h
@@ -29,7 +29,7 @@ class UCameraShakeCameraNode
 protected:
 
 	// UCameraNode interface.
-	virtual void OnPreBuild(FCameraBuildLog& BuildLog) override;
+	virtual void OnPreBuild(FCameraBuildContext& BuildContext) override;
 	virtual void OnBuild(FCameraObjectBuildContext& BuildContext) override;
 	virtual FCameraNodeEvaluatorPtr OnBuildEvaluator(FCameraNodeEvaluatorBuilder& Builder) const override;
 
diff --git a/Engine/Plugins/Cameras/GameplayCameras/Source/GameplayCameras/Public/Nodes/Utility/BlueprintCameraNode.h b/Engine/Plugins/Cameras/GameplayCameras/Source/GameplayCameras/Public/Nodes/Utility/BlueprintCameraNode.h
--- a/Engine/Plugins/Cameras/GameplayCameras/Source/GameplayCameras/Public/Nodes/Utility/BlueprintCameraNode.h
+++ b/Engine/Plugins/Cameras/GameplayCameras/Source/GameplayCameras/Public/Nodes/Utility/BlueprintCameraNode.h
@@ -158,7 +158,7 @@ public:
 protected:
 
 	// UCameraNode interface.
-	virtual void OnPreBuild(FCameraBuildLog& BuildLog) override;
+	virtual void OnPreBuild(FCameraBuildContext& BuildContext) override;
 	virtual void OnBuild(FCameraObjectBuildContext& BuildContext) override;
 	virtual FCameraNodeEvaluatorPtr OnBuildEvaluator(FCameraNodeEvaluatorBuilder& Builder) const override;
 
diff --git a/Engine/Plugins/Cameras/GameplayCameras/Source/GameplayCamerasEditor/GameplayCamerasEditor.Build.cs b/Engine/Plugins/Cameras/GameplayCameras/Source/GameplayCamerasEditor/GameplayCamerasEditor.Build.cs
--- a/Engine/Plugins/Cameras/GameplayCameras/Source/GameplayCamerasEditor/GameplayCamerasEditor.Build.cs
+++ b/Engine/Plugins/Cameras/GameplayCameras/Source/GameplayCamerasEditor/GameplayCamerasEditor.Build.cs
@@ -33,6 +33,7 @@ public class GameplayCamerasEditor : ModuleRules
 				"EditorSubsystem",
 				"Engine",
 				"GameplayCameras",
+				"GameplayCamerasUncookedOnly",
 				"GraphEditor",
 				"InputCore",
 				"InteractiveToolsFramework",
diff --git a/Engine/Plugins/Cameras/GameplayCameras/Source/GameplayCamerasEditor/Private/Editors/SCameraAssetInterfaceParametersPanel.cpp b/Engine/Plugins/Cameras/GameplayCameras/Source/GameplayCamerasEditor/Private/Editors/SCameraAssetInterfaceParametersPanel.cpp
new file mode 100644
--- /dev/null
+++ b/Engine/Plugins/Cameras/GameplayCameras/Source/GameplayCamerasEditor/Private/Editors/SCameraAssetInterfaceParametersPanel.cpp
@@ -0,0 +1,636 @@
+// Copyright Epic Games, Inc. All Rights Reserved.
+
+#include "Editors/SCameraAssetInterfaceParametersPanel.h"
+
+#include "AssetRegistry/ARFilter.h"
+#include "Core/CameraAsset.h"
+#include "Core/CameraDirector.h"
+#include "Core/CameraRigAsset.h"
+#include "Framework/MultiBox/MultiBoxBuilder.h"
+#include "Framework/Views/ITypedTableView.h"
+#include "Helpers/CameraContextDataPinTypeHelper.h"
+#include "Helpers/CameraVariablePinTypeHelper.h"
+#include "IContentBrowserSingleton.h"
+#include "Kismet2/BlueprintEditorUtils.h"
+#include "ScopedTransaction.h"
+#include "Toolkits/CameraAssetEditorToolkit.h"
+#include "Widgets/Input/SCheckBox.h"
+#include "Widgets/Input/SComboBox.h"
+#include "Widgets/Text/SInlineEditableTextBlock.h"
+#include "Widgets/Text/SRichTextBlock.h"
+#include "Widgets/Text/STextBlock.h"
+#include "Widgets/Views/SListView.h"
+
+#define LOCTEXT_NAMESPACE "SCameraAssetInterfaceParametersPanel"
+
+namespace UE::Cameras
+{
+
+class SCameraAssetSourceCameraRigPicker : public SCompoundWidget
+{
+public:
+
+	SLATE_BEGIN_ARGS(SCameraAssetSourceCameraRigPicker)
+	{}
+		SLATE_ARGUMENT(UCameraAssetInterfaceParameter*, InterfaceParameter)
+	SLATE_END_ARGS()
+
+	void Construct(const FArguments& Args)
+	{
+		InterfaceParameter = Args._InterfaceParameter;
+
+		CameraAsset = InterfaceParameter->GetTypedOuter<UCameraAsset>();
+		if (ensure(CameraAsset))
+		{
+			if (UCameraDirector* CameraDirector = CameraAsset->GetCameraDirector())
+			{
+				FCameraDirectorRigUsageInfo UsageInfo;
+				CameraDirector->GatherRigUsageInfo(UsageInfo);
+				for (const UCameraRigAsset* CameraRig : UsageInfo.CameraRigs)
+				{
+					ReferencedCameraRigPaths.Add(CameraRig->GetPathName());
+				}
+			}
+		}
+		
+		ChildSlot
+		[
+			SAssignNew(ComboButton, SComboButton)
+			.OnGetMenuContent(this, &SCameraAssetSourceCameraRigPicker::OnGenerateAssetPicker)
+			.ButtonContent()
+			[
+				SNew(STextBlock)
+				.Text(this, &SCameraAssetSourceCameraRigPicker::GetSourceCameraRigName)
+			]
+		];
+	}
+
+private:
+
+	TSharedRef<SWidget> OnGenerateAssetPicker()
+	{
+		FAssetPickerConfig AssetPickerConfig;
+		AssetPickerConfig.bAllowNullSelection = false;
+		AssetPickerConfig.InitialAssetViewType = EAssetViewType::List;
+		AssetPickerConfig.SelectionMode = ESelectionMode::Single;
+		AssetPickerConfig.Filter.ClassPaths.Add(UCameraRigAsset::StaticClass()->GetClassPathName());
+		AssetPickerConfig.OnShouldFilterAsset = FOnShouldFilterAsset::CreateSP(this, &SCameraAssetSourceCameraRigPicker::OnShouldFilterAsset);
+		AssetPickerConfig.OnAssetSelected = FOnAssetSelected::CreateSP(this, &SCameraAssetSourceCameraRigPicker::OnAssetSelected);
+		AssetPickerConfig.RefreshAssetViewDelegates.Add(&RefreshAssetPickerEvent);
+
+		IContentBrowserSingleton& ContentBrowser = IContentBrowserSingleton::Get();
+		TSharedRef<SWidget> AssetPicker = ContentBrowser.CreateAssetPicker(AssetPickerConfig);
+
+		return SNew(SVerticalBox)
+			+SVerticalBox::Slot()
+			.FillHeight(1.f)
+			.MinHeight(400.f)
+			[
+				AssetPicker
+			]
+			+SVerticalBox::Slot()
+			.AutoHeight()
+			.Padding(8.f)
+			[
+				SNew(SCheckBox)
+				.IsChecked(this, &SCameraAssetSourceCameraRigPicker::IsShowingOnlyReferencedCameraRigs)
+				.OnCheckStateChanged(this, &SCameraAssetSourceCameraRigPicker::OnToggleOnlyShowReferencedCameraRigs)
+				[
+					SNew(STextBlock)
+					.Text(LOCTEXT("ShowOnlyReferencedCameraRigs", "Show only referenced camera rigs"))
+				]
+			];
+	}
+
+	FText GetSourceCameraRigName() const
+	{
+		if (UCameraRigAsset* CameraRig = InterfaceParameter->SourceCameraRig.Get())
+		{
+			return FText::FromName(CameraRig->GetFName());
+		}
+		return LOCTEXT("NoSourceCameraRig", "None");
+	}
+
+	bool OnShouldFilterAsset(const FAssetData& AssetData)
+	{
+		if (bShowOnlyReferencedCameraRigs)
+		{
+			return !ReferencedCameraRigPaths.Contains(AssetData.GetObjectPathString());
+		}
+		return false;
+	}
+
+	void OnAssetSelected(const FAssetData& SelectedAsset)
+	{
+		UCameraRigAsset* SelectedCameraRig = Cast<UCameraRigAsset>(SelectedAsset.GetAsset());
+		if (SelectedCameraRig != InterfaceParameter->SourceCameraRig)
+		{
+			const FScopedTransaction Transaction(LOCTEXT("SetSourceCameraRig", "Set Source Camera Rig"));
+
+			InterfaceParameter->Modify();
+			InterfaceParameter->SourceCameraRig = SelectedCameraRig;
+		}
+
+		ComboButton->SetIsOpen(false);
+	}
+
+	ECheckBoxState IsShowingOnlyReferencedCameraRigs() const
+	{
+		return (bShowOnlyReferencedCameraRigs ? ECheckBoxState::Checked : ECheckBoxState::Unchecked);
+	}
+
+	void OnToggleOnlyShowReferencedCameraRigs(ECheckBoxState CheckBoxState)
+	{
+		bShowOnlyReferencedCameraRigs = (CheckBoxState == ECheckBoxState::Checked);
+
+		RefreshAssetPickerEvent.ExecuteIfBound(false);
+	}
+
+private:
+
+	UCameraAssetInterfaceParameter* InterfaceParameter = nullptr;
+
+	UCameraAsset* CameraAsset = nullptr;
+	TSet<FString> ReferencedCameraRigPaths;
+
+	TSharedPtr<SComboButton> ComboButton;
+	FRefreshAssetViewDelegate RefreshAssetPickerEvent;
+	bool bShowOnlyReferencedCameraRigs = true;
+};
+
+class SCameraAssetSourceParameterNamePicker : public SCompoundWidget
+{
+	SLATE_BEGIN_ARGS(SCameraAssetSourceParameterNamePicker)
+	{}
+		SLATE_ARGUMENT(UCameraAssetInterfaceParameter*, InterfaceParameter)
+	SLATE_END_ARGS()
+
+	void Construct(const FArguments& Args)
+	{
+		InterfaceParameter = Args._InterfaceParameter;
+		CameraAsset = InterfaceParameter->GetTypedOuter<UCameraAsset>();
+
+		EnsureSourceParameterNamesCached();
+
+		ChildSlot
+		[
+			SNew(SComboBox<TSharedPtr<FName>>)
+			.OptionsSource(&LastKnownSourceParameterNames)
+			.OnGenerateWidget(this, &SCameraAssetSourceParameterNamePicker::OnGenerateParameterNameWidget)
+			.OnSelectionChanged(this, &SCameraAssetSourceParameterNamePicker::OnParameterNameSelectionChanged)
+			[
+				SNew(STextBlock)
+				.ColorAndOpacity(this, &SCameraAssetSourceParameterNamePicker::GetSourceParameterNameColorAndOpacity)
+				.Text(this, &SCameraAssetSourceParameterNamePicker::GetSourceParameterName)
+			]
+		];
+	}
+
+protected:
+
+	void Tick(const FGeometry& AllottedGeometry, const double InCurrentTime, const float InDeltaTime) override
+	{
+		SCompoundWidget::Tick(AllottedGeometry, InCurrentTime, InDeltaTime);
+
+		// Keep checking in case the source camera rig changes.
+		EnsureSourceParameterNamesCached();
+	}
+
+private:
+
+	TSharedRef<SWidget> OnGenerateParameterNameWidget(TSharedPtr<FName> Item)
+	{
+		const FName ItemName = Item ? *Item : NAME_None;
+		return SNew(STextBlock)
+			.Text(FText::FromName(ItemName));
+	}
+
+	void OnParameterNameSelectionChanged(TSharedPtr<FName> Item, ESelectInfo::Type SelectInfo)
+	{
+		FName NewName = Item ? *Item : NAME_None;
+		if (InterfaceParameter->SourceParameterName != NewName)
+		{
+			const FScopedTransaction Transaction(LOCTEXT("SetSourceParameterName", "Set Source Parameter Name"));
+
+			InterfaceParameter->Modify();
+			InterfaceParameter->SourceParameterName = NewName;
+		}
+	}
+
+	FSlateColor GetSourceParameterNameColorAndOpacity() const
+	{
+		bool bIsValidParameter = false;
+		bool bIsParameterNone = InterfaceParameter->SourceParameterName.IsNone();
+		if (LastKnownSourceCameraRig && !bIsParameterNone)
+		{
+			FCameraObjectInterfaceParameterDefinition ParameterDefinition;
+			bIsValidParameter = LastKnownSourceCameraRig->FindParameterDefinitionByName(
+					InterfaceParameter->SourceParameterName, ParameterDefinition);
+		}
+		if (bIsValidParameter)
+		{
+			return FSlateColor::UseForeground();
+		}
+		else if (bIsParameterNone)
+		{
+			return FSlateColor::UseSubduedForeground();
+		}
+		else
+		{
+			return FStyleColors::Error;
+		}
+	}
+
+	FText GetSourceParameterName() const
+	{
+		return FText::FromName(InterfaceParameter->SourceParameterName);
+	}
+
+	void EnsureSourceParameterNamesCached()
+	{
+		UCameraRigAsset* CameraRig = InterfaceParameter->SourceCameraRig.Get();
+		if (CameraRig == LastKnownSourceCameraRig)
+		{
+			return;
+		}
+
+		LastKnownSourceCameraRig = CameraRig;
+		LastKnownSourceParameterNames.Reset();
+		if (CameraRig)
+		{
+			for (const FCameraObjectInterfaceParameterDefinition& ParameterDefinition : CameraRig->GetParameterDefinitions())
+			{
+				if (ParameterDefinition.bIsVisible)
+				{
+					LastKnownSourceParameterNames.Add(MakeShared<FName>(ParameterDefinition.ParameterName));
+				}
+			}
+		}
+		if (LastKnownSourceParameterNames.IsEmpty())
+		{
+			LastKnownSourceParameterNames.Add(MakeShared<FName>(NAME_None));
+		}
+	}
+
+private:
+
+	UCameraAsset* CameraAsset = nullptr;
+	UCameraAssetInterfaceParameter* InterfaceParameter = nullptr;
+
+	TArray<TSharedPtr<FName>> LastKnownSourceParameterNames;
+	UCameraRigAsset* LastKnownSourceCameraRig = nullptr;
+};
+
+class SCameraAssetInterfaceParameterTableRow : public SMultiColumnTableRow<TObjectPtr<UCameraAssetInterfaceParameter>>
+{
+public:
+
+	SLATE_BEGIN_ARGS(SCameraAssetInterfaceParameterTableRow)
+	{}
+		SLATE_ARGUMENT(TObjectPtr<UCameraAssetInterfaceParameter>, Item)
+	SLATE_END_ARGS()
+
+	typedef SMultiColumnTableRow<TObjectPtr<UCameraAssetInterfaceParameter>> FSuperRowType;
+	typedef typename STableRow<TObjectPtr<UCameraAssetInterfaceParameter>>::FArguments FTableRowArgs;
+		
+	void Construct(const FArguments& Args, const TSharedRef<STableViewBase>& OwnerTable)
+	{
+		Item = Args._Item;
+
+		EnsurePinTypeCached();
+
+		FSuperRowType::Construct(
+				FTableRowArgs(),
+				OwnerTable);
+	}
+
+	void EnterNameEditingMode()
+	{
+		NameTextBlock->EnterEditingMode();
+	}
+	
+protected:
+
+	void Tick(const FGeometry& AllottedGeometry, const double InCurrentTime, const float InDeltaTime) override
+	{
+		FSuperRowType::Tick(AllottedGeometry, InCurrentTime, InDeltaTime);
+
+		// Keep checking in case the source camera rig or parameter name change.
+		EnsurePinTypeCached();
+	}
+
+private:
+
+	virtual TSharedRef<SWidget> GenerateWidgetForColumn(const FName& InColumnName) override
+	{
+		if (InColumnName == SCameraAssetInterfaceParametersPanel::ParameterTypeColumn)
+		{
+			return SNew(SHorizontalBox)
+				+SHorizontalBox::Slot()
+				.AutoWidth()
+				.VAlign(VAlign_Center)
+				.Padding(4.f, 2.f)
+				[
+					SNew(SImage)
+						.Image(this, &SCameraAssetInterfaceParameterTableRow::GetTypeIcon)
+						.ColorAndOpacity(this, &SCameraAssetInterfaceParameterTableRow::GetTypeIconColor)
+				];
+		}
+		else if (InColumnName == SCameraAssetInterfaceParametersPanel::ParameterNameColumn)
+		{
+			return SAssignNew(NameTextBlock, SInlineEditableTextBlock)
+				.IsSelected(this, &SCameraAssetInterfaceParameterTableRow::IsSelected)
+				.Text_Lambda([this]() { return FText::FromString(Item->InterfaceParameterName); })
+				.OnTextCommitted(this, &SCameraAssetInterfaceParameterTableRow::OnParameterNameTextCommitted)
+				.ToolTipText(this, &SCameraAssetInterfaceParameterTableRow::GetParameterToolTip);
+		}
+		else if (InColumnName == SCameraAssetInterfaceParametersPanel::SourceCameraRigColumn)
+		{
+			return SNew(SCameraAssetSourceCameraRigPicker).InterfaceParameter(Item);
+		}
+		else if (InColumnName == SCameraAssetInterfaceParametersPanel::SourceParameterNameColumn)
+		{
+			return SNew(SCameraAssetSourceParameterNamePicker).InterfaceParameter(Item);
+		}
+
+		return SNullWidget::NullWidget;
+	}
+
+	void OnParameterNameTextCommitted(const FText& Text, ETextCommit::Type CommitType)
+	{
+		const FString NewParameterName = Text.ToString();
+
+		if (Item->InterfaceParameterName != NewParameterName)
+		{
+			const FScopedTransaction Transaction(LOCTEXT("RenameInterfaceParameterTransaction", "Rename Interface Parameter"));
+
+			Item->Modify();
+			Item->InterfaceParameterName = NewParameterName;
+		}
+	}
+
+	FText GetParameterToolTip() const
+	{
+		const FGuid Guid = Item->GetGuid();
+		return FText::FromString(Guid.ToString(EGuidFormats::DigitsWithHyphensInParentheses));
+	}
+
+	const FSlateBrush* GetTypeIcon() const
+	{
+		return PinTypeIcon;
+	}
+
+	FSlateColor GetTypeIconColor() const
+	{
+		return PinTypeIconColor;
+	}
+
+	void EnsurePinTypeCached()
+	{
+		UCameraRigAsset* CameraRig = Item->SourceCameraRig.Get();
+		if (CameraRig == LastKnownSourceCameraRig && Item->SourceParameterName == LastKnownSourceParameterName)
+		{
+			return;
+		}
+
+		const FSlateBrush* InvalidIcon = FAppStyle::GetBrush("Icons.Warning.Solid");
+		const FLinearColor DefaultIconColor = FLinearColor::White;
+
+		LastKnownSourceCameraRig = CameraRig;
+		LastKnownSourceParameterName = Item->SourceParameterName;
+
+		if (!CameraRig || Item->SourceParameterName.IsNone())
+		{
+			PinTypeIcon = InvalidIcon;
+			PinTypeIconColor = DefaultIconColor;
+			return;
+		}
+
+		FCameraObjectInterfaceParameterDefinition ParameterDefinition;
+		const bool bFoundSourceParameter = CameraRig->FindParameterDefinitionByName(Item->SourceParameterName, ParameterDefinition);
+		if (!bFoundSourceParameter)
+		{
+			PinTypeIcon = InvalidIcon;
+			PinTypeIconColor = DefaultIconColor;
+			return;
+		}
+
+		FEdGraphPinType PinType;
+		if (ParameterDefinition.ParameterType == ECameraObjectInterfaceParameterType::Blendable)
+		{
+			PinType = FCameraVariablePinTypeHelper::GetPinType(ParameterDefinition.VariableType, ParameterDefinition.BlendableStructType);
+		}
+		else if (ParameterDefinition.ParameterType == ECameraObjectInterfaceParameterType::Data)
+		{
+			PinType = FCameraContextDataPinTypeHelper::GetPinType(ParameterDefinition.DataType, ParameterDefinition.DataContainerType, ParameterDefinition.DataTypeObject);
+		}
+
+		PinTypeIcon = FBlueprintEditorUtils::GetIconFromPin(PinType);
+		PinTypeIconColor = GetDefault<UEdGraphSchema_K2>()->GetPinTypeColor(PinType);
+	}
+
+protected:
+
+	TObjectPtr<UCameraAssetInterfaceParameter> Item;
+	TSharedPtr<SInlineEditableTextBlock> NameTextBlock;
+
+	UCameraRigAsset* LastKnownSourceCameraRig = nullptr;
+	FName LastKnownSourceParameterName;
+
+	const FSlateBrush* PinTypeIcon = nullptr;
+	FLinearColor PinTypeIconColor = FLinearColor::White;
+};
+
+const FName SCameraAssetInterfaceParametersPanel::ParameterTypeColumn(TEXT("ParameterType"));
+const FName SCameraAssetInterfaceParametersPanel::ParameterNameColumn(TEXT("ParameterName"));
+const FName SCameraAssetInterfaceParametersPanel::SourceCameraRigColumn(TEXT("SourceCameraRig"));
+const FName SCameraAssetInterfaceParametersPanel::SourceParameterNameColumn(TEXT("SourceParameterName"));
+
+void SCameraAssetInterfaceParametersPanel::Construct(const FArguments& Args, FCameraAssetEditorToolkit* OwnerToolkit)
+{
+	CameraAsset = OwnerToolkit->GetCameraAsset();
+	CameraAsset->EventHandlers.Register(EventHandler, this);
+
+	Toolkit = OwnerToolkit;
+
+	ChildSlot
+	[
+		SNew(SVerticalBox)
+		+SVerticalBox::Slot()
+		.AutoHeight()
+		[
+			SNew(SBorder)
+			.BorderImage(FAppStyle::GetBrush("ToolPanel.GroupBorder"))
+			.Padding(FMargin(3, 5))
+			[
+				SNew(SHorizontalBox)
+				+SHorizontalBox::Slot()
+				.FillWidth(1.f)
+				.VAlign(VAlign_Center)
+				.Padding(5.f)
+				[
+					SNew(SRichTextBlock)
+					.Text(LOCTEXT("ParametersPanelTitle", "Parameters"))
+					.TransformPolicy(ETextTransformPolicy::ToUpper)
+					.DecoratorStyleSet(&FAppStyle::Get())
+					.TextStyle(FAppStyle::Get(), "DetailsView.CategoryTextStyle")
+				]
+				+SHorizontalBox::Slot()
+				.AutoWidth()
+				.VAlign(VAlign_Center)
+				[
+					SNew(SButton)
+					.ButtonStyle(FAppStyle::Get(), "SimpleButton")
+					.ContentPadding(FMargin(1, 0))
+					.ToolTipText(LOCTEXT("AddParameterToolTip", "Add a parameter"))
+					.OnClicked(this, &SCameraAssetInterfaceParametersPanel::OnAddInterfaceParameter)
+					[
+						SNew(SImage)
+						.Image(FAppStyle::Get().GetBrush("Icons.PlusCircle"))
+						.ColorAndOpacity(FSlateColor::UseForeground())
+					]
+				]
+			]
+		]
+		+SVerticalBox::Slot()
+		.FillHeight(1.f)
+		[
+			SAssignNew(ParametersListView, SListView<TObjectPtr<UCameraAssetInterfaceParameter>>)
+			.ListItemsSource(&CameraAsset->Interface.Parameters)
+			.OnGenerateRow(this, &SCameraAssetInterfaceParametersPanel::OnGenerateInterfaceParameterRow)
+			.OnSelectionChanged(this, &SCameraAssetInterfaceParametersPanel::OnInterfaceParameterSelectionChanged)
+			.OnContextMenuOpening(this, &SCameraAssetInterfaceParametersPanel::OnInterfaceParameterContextMenuOpening)
+			.HeaderRow
+			(
+				SNew(SHeaderRow)
+
+				+SHeaderRow::Column(ParameterTypeColumn)
+				.FixedWidth(64)
+				.DefaultLabel(LOCTEXT("ParameterTypeColumnLabel", "Type"))
+
+				+SHeaderRow::Column(ParameterNameColumn)
+				.FillWidth(0.3f)
+				.DefaultLabel(LOCTEXT("ParameterNameColumnLabel", "Name"))
+
+				+SHeaderRow::Column(SourceCameraRigColumn)
+				.FillWidth(0.3f)
+				.DefaultLabel(LOCTEXT("SourceCameraRigColumnLabel", "Source Camera Rig"))
+
+				+SHeaderRow::Column(SourceParameterNameColumn)
+				.FillWidth(0.3f)
+				.DefaultLabel(LOCTEXT("SourceParameterNameColumnLabel", "Source Parameter"))
+			)
+		]
+	];
+}
+
+void SCameraAssetInterfaceParametersPanel::RequestListRefresh()
+{
+	bListRefreshRequested = true;
+}
+
+void SCameraAssetInterfaceParametersPanel::Tick(const FGeometry& AllottedGeometry, const double InCurrentTime, const float InDeltaTime)
+{
+	SCompoundWidget::Tick(AllottedGeometry, InCurrentTime, InDeltaTime);
+
+	if (bListRefreshRequested)
+	{
+		bListRefreshRequested = false;
+
+		ParametersListView->RequestListRefresh();
+	}
+}
+
+TSharedRef<ITableRow> SCameraAssetInterfaceParametersPanel::OnGenerateInterfaceParameterRow(TObjectPtr<UCameraAssetInterfaceParameter> Item, const TSharedRef<STableViewBase>& OwnerTable)
+{
+	return SNew(SCameraAssetInterfaceParameterTableRow, OwnerTable)
+		.Item(Item);
+}
+
+void SCameraAssetInterfaceParametersPanel::OnInterfaceParameterSelectionChanged(TObjectPtr<UCameraAssetInterfaceParameter> Item, ESelectInfo::Type Type)
+{
+	OnInterfaceParameterSelectedDelegate.Broadcast(Item);
+}
+
+TSharedPtr<SWidget> SCameraAssetInterfaceParametersPanel::OnInterfaceParameterContextMenuOpening()
+{
+	TArray<TObjectPtr<UCameraAssetInterfaceParameter>> SelectedItems = ParametersListView->GetSelectedItems();
+	if (SelectedItems.Num() == 0)
+	{
+		return nullptr;
+	}
+
+	FMenuBuilder MenuBuilder(true, nullptr);
+	TObjectPtr<UCameraAssetInterfaceParameter> Item = SelectedItems[0];
+
+	MenuBuilder.BeginSection(NAME_None, LOCTEXT("InterfaceParameterEditSection", "Edit Interface Parameter"));
+	{
+		MenuBuilder.AddMenuEntry(
+				LOCTEXT("RenameInterfaceParameter", "Rename"),
+				LOCTEXT("RenameInterfaceParameterToolTip", "Renames this interface parameter"),
+				FSlateIcon(FAppStyle::GetAppStyleSetName(), "GenericCommands.Rename"),
+				FExecuteAction::CreateSP(this, &SCameraAssetInterfaceParametersPanel::OnRenameInterfaceParameter, Item));
+		MenuBuilder.AddMenuEntry(
+				LOCTEXT("DeleteInterfaceParameter", "Delete"),
+				LOCTEXT("DeleteInterfaceParameterToolTip", "Deletes this interface parameter"),
+				FSlateIcon(FAppStyle::GetAppStyleSetName(), "GenericCommands.Delete"),
+				FExecuteAction::CreateSP(this, &SCameraAssetInterfaceParametersPanel::OnDeleteInterfaceParameter, Item));
+	}
+	MenuBuilder.EndSection();
+
+	return MenuBuilder.MakeWidget();
+}
+
+FReply SCameraAssetInterfaceParametersPanel::OnAddInterfaceParameter()
+{
+	const FScopedTransaction Transaction(LOCTEXT("AddInterfaceParameter", "Add Interface Parameter"));
+
+	CameraAsset->Modify();
+
+	UCameraAssetInterfaceParameter* NewParameter = NewObject<UCameraAssetInterfaceParameter>(CameraAsset, NAME_None, RF_Transactional);
+	NewParameter->InterfaceParameterName = NewParameter->GetName();
+
+	CameraAsset->Interface.Parameters.Add(NewParameter);
+	CameraAsset->EventHandlers.Notify(&ICameraAssetEventHandler::OnCameraAssetInterfaceChanged);
+
+	ParametersListView->RequestListRefresh();
+
+	return FReply::Handled();
+}
+
+void SCameraAssetInterfaceParametersPanel::OnRenameInterfaceParameter(TObjectPtr<UCameraAssetInterfaceParameter> Item)
+{
+	TSharedPtr<ITableRow> RowWidget = ParametersListView->WidgetFromItem(Item);
+	if (!ensure(RowWidget))
+	{
+		return;
+	}
+
+	TSharedPtr<SCameraAssetInterfaceParameterTableRow> TypedRowWidget = 
+		StaticCastSharedPtr<SCameraAssetInterfaceParameterTableRow>(RowWidget);
+	TypedRowWidget->EnterNameEditingMode();
+}
+
+void SCameraAssetInterfaceParametersPanel::OnDeleteInterfaceParameter(TObjectPtr<UCameraAssetInterfaceParameter> Item)
+{
+	const FScopedTransaction Transaction(LOCTEXT("RemoveInterfaceParameter", "Remove Interface Parameter"));
+
+	CameraAsset->Modify();
+
+	const int32 NumRemoved = CameraAsset->Interface.Parameters.Remove(Item);
+	ensure(NumRemoved == 1);
+
+	CameraAsset->EventHandlers.Notify(&ICameraAssetEventHandler::OnCameraAssetInterfaceChanged);
+
+	ParametersListView->RequestListRefresh();
+}
+
+void SCameraAssetInterfaceParametersPanel::OnCameraAssetInterfaceChanged()
+{
+	bListRefreshRequested = true;
+}
+
+}  // namespace UE::Cameras
+
+#undef LOCTEXT_NAMESPACE
+
diff --git a/Engine/Plugins/Cameras/GameplayCameras/Source/GameplayCamerasEditor/Private/Editors/SCameraObjectInterfaceParametersPanel.cpp b/Engine/Plugins/Cameras/GameplayCameras/Source/GameplayCamerasEditor/Private/Editors/SCameraObjectInterfaceParametersPanel.cpp
--- a/Engine/Plugins/Cameras/GameplayCameras/Source/GameplayCamerasEditor/Private/Editors/SCameraObjectInterfaceParametersPanel.cpp
+++ b/Engine/Plugins/Cameras/GameplayCameras/Source/GameplayCamerasEditor/Private/Editors/SCameraObjectInterfaceParametersPanel.cpp
@@ -4,6 +4,8 @@
 
 #include "Editors/CameraNodeGraphDragDropOp.h"
 #include "Framework/MultiBox/MultiBoxBuilder.h"
+#include "Helpers/CameraContextDataPinTypeHelper.h"
+#include "Helpers/CameraVariablePinTypeHelper.h"
 #include "IGameplayCamerasModule.h"
 #include "SPinTypeSelector.h"
 #include "ScopedTransaction.h"
@@ -25,9 +27,6 @@
 namespace UE::Cameras
 {
 
-static const FName ParameterTypeColumn(TEXT("ParameterType"));
-static const FName ParameterNameColumn(TEXT("ParameterName"));
-
 /**
  * List entry for any interface parameter panel.
  */
@@ -62,7 +61,7 @@ protected:
 
 	virtual TSharedRef<SWidget> GenerateWidgetForColumn(const FName& InColumnName) override
 	{
-		if (InColumnName == ParameterNameColumn)
+		if (InColumnName == SCameraObjectInterfaceParametersPanel::ParameterNameColumn)
 		{
 			return SAssignNew(NameTextBlock, SInlineEditableTextBlock)
 				.IsSelected(this, &SCameraObjectInterfaceParameterTableRowBase<ParameterType>::IsSelected)
@@ -119,7 +118,7 @@ protected:
 
 	virtual TSharedRef<SWidget> GenerateWidgetForColumn(const FName& InColumnName) override
 	{
-		if (InColumnName == ParameterTypeColumn)
+		if (InColumnName == SCameraObjectInterfaceParametersPanel::ParameterTypeColumn)
 		{
 			TSharedRef<FGameplayCamerasEditorStyle> GameplayCamerasEditorStyle = FGameplayCamerasEditorStyle::Get();
 
@@ -196,67 +195,8 @@ protected:
 
 	FEdGraphPinType GetBlendableParameterPinType() const
 	{
-		FEdGraphPinType PinType;
-		switch (Item->VariableType)
-		{
-			case ECameraVariableType::Boolean:
-				PinType.PinCategory = UEdGraphSchema_K2::PC_Boolean;
-				break;
-			case ECameraVariableType::Integer32:
-				PinType.PinCategory = UEdGraphSchema_K2::PC_Int;
-				break;
-			case ECameraVariableType::Float:
-				PinType.PinCategory = UEdGraphSchema_K2::PC_Float;
-				break;
-			case ECameraVariableType::Double:
-				PinType.PinCategory = UEdGraphSchema_K2::PC_Double;
-				break;
-			case ECameraVariableType::Vector2f:
-				PinType.PinCategory = UEdGraphSchema_K2::PC_Struct;
-				PinType.PinSubCategoryObject = TVariantStructure<FVector2f>::Get();
-				break;
-			case ECameraVariableType::Vector2d:
-				PinType.PinCategory = UEdGraphSchema_K2::PC_Struct;
-				PinType.PinSubCategoryObject = TBaseStructure<FVector2D>::Get();
-				break;
-			case ECameraVariableType::Vector3f:
-				PinType.PinCategory = UEdGraphSchema_K2::PC_Struct;
-				PinType.PinSubCategoryObject = TVariantStructure<FVector3f>::Get();
-				break;
-			case ECameraVariableType::Vector3d:
-				PinType.PinCategory = UEdGraphSchema_K2::PC_Struct;
-				PinType.PinSubCategoryObject = TBaseStructure<FVector>::Get();
-				break;
-			case ECameraVariableType::Vector4f:
-				PinType.PinCategory = UEdGraphSchema_K2::PC_Struct;
-				PinType.PinSubCategoryObject = TVariantStructure<FVector4f>::Get();
-				break;
-			case ECameraVariableType::Vector4d:
-				PinType.PinCategory = UEdGraphSchema_K2::PC_Struct;
-				PinType.PinSubCategoryObject = TBaseStructure<FVector4>::Get();
-				break;
-			case ECameraVariableType::Rotator3f:
-				PinType.PinCategory = UEdGraphSchema_K2::PC_Struct;
-				PinType.PinSubCategoryObject = TVariantStructure<FRotator3f>::Get();
-				break;
-			case ECameraVariableType::Rotator3d:
-				PinType.PinCategory = UEdGraphSchema_K2::PC_Struct;
-				PinType.PinSubCategoryObject = TBaseStructure<FRotator>::Get();
-				break;
-			case ECameraVariableType::Transform3f:
-				PinType.PinCategory = UEdGraphSchema_K2::PC_Struct;
-				PinType.PinSubCategoryObject = TVariantStructure<FTransform3f>::Get();
-				break;
-			case ECameraVariableType::Transform3d:
-				PinType.PinCategory = UEdGraphSchema_K2::PC_Struct;
-				PinType.PinSubCategoryObject = TBaseStructure<FTransform>::Get();
-				break;
-			case ECameraVariableType::BlendableStruct:
-				PinType.PinCategory = UEdGraphSchema_K2::PC_Struct;
-				PinType.PinSubCategoryObject = const_cast<UScriptStruct*>(Item->BlendableStructType.Get());
-				break;
-		}
-		return PinType;
+		return FCameraVariablePinTypeHelper::GetPinType(
+				Item->VariableType, Item->BlendableStructType.Get());
 	}
 
 	void OnBlendableParameterPinTypeChanged(const FEdGraphPinType& PinType)
@@ -400,7 +340,7 @@ protected:
 
 	virtual TSharedRef<SWidget> GenerateWidgetForColumn(const FName& InColumnName) override
 	{
-		if (InColumnName == ParameterTypeColumn)
+		if (InColumnName == SCameraObjectInterfaceParametersPanel::ParameterTypeColumn)
 		{
 			TSharedRef<FGameplayCamerasEditorStyle> GameplayCamerasEditorStyle = FGameplayCamerasEditorStyle::Get();
 
@@ -432,39 +372,8 @@ protected:
 
 	FEdGraphPinType GetDataParameterPinType() const
 	{
-		FEdGraphPinType PinType;
-		PinType.PinSubCategoryObject = const_cast<UObject*>(Item->DataTypeObject.Get());
-
-		switch (Item->DataType)
-		{
-			case ECameraContextDataType::Name:
-				PinType.PinCategory = UEdGraphSchema_K2::PC_Name;
-				break;
-			case ECameraContextDataType::String:
-				PinType.PinCategory = UEdGraphSchema_K2::PC_String;
-				break;
-			case ECameraContextDataType::Enum:
-				PinType.PinCategory = UEdGraphSchema_K2::PC_Enum;
-				break;
-			case ECameraContextDataType::Struct:
-				PinType.PinCategory = UEdGraphSchema_K2::PC_Struct;
-				break;
-			case ECameraContextDataType::Object:
-				PinType.PinCategory = UEdGraphSchema_K2::PC_Object;
-				break;
-			case ECameraContextDataType::Class:
-				PinType.PinCategory = UEdGraphSchema_K2::PC_Class;
-				break;
-		}
-
-		switch (Item->DataContainerType)
-		{
-			case ECameraContextDataContainerType::Array:
-				PinType.ContainerType = EPinContainerType::Array;
-				break;
-		}
-
-		return PinType;
+		return FCameraContextDataPinTypeHelper::GetPinType(
+				Item->DataType, Item->DataContainerType, Item->DataTypeObject.Get());
 	}
 
 	void OnDataParameterPinTypeChanged(const FEdGraphPinType& PinType)
@@ -528,6 +437,7 @@ protected:
 			Item->Modify();
 
 			Item->DataType = NewDataType;
+			Item->DataContainerType = NewDataContainerType;
 			Item->DataTypeObject = NewDataTypeObject;
 		}
 	}
@@ -538,6 +448,9 @@ protected:
 	}
 };
 
+const FName SCameraObjectInterfaceParametersPanel::ParameterTypeColumn(TEXT("ParameterType"));
+const FName SCameraObjectInterfaceParametersPanel::ParameterNameColumn(TEXT("ParameterName"));
+
 void SCameraObjectInterfaceParametersPanel::Construct(const FArguments& Args, FCameraObjectInterfaceParametersToolkit* OwnerToolkit)
 {
 	CameraObject = OwnerToolkit->GetCameraObject();
diff --git a/Engine/Plugins/Cameras/GameplayCameras/Source/GameplayCamerasEditor/Private/GameplayCamerasLiveEditManager.cpp b/Engine/Plugins/Cameras/GameplayCameras/Source/GameplayCamerasEditor/Private/GameplayCamerasLiveEditManager.cpp
--- a/Engine/Plugins/Cameras/GameplayCameras/Source/GameplayCamerasEditor/Private/GameplayCamerasLiveEditManager.cpp
+++ b/Engine/Plugins/Cameras/GameplayCameras/Source/GameplayCamerasEditor/Private/GameplayCamerasLiveEditManager.cpp
@@ -3,6 +3,7 @@
 #include "GameplayCamerasLiveEditManager.h"
 
 #include "Build/CameraAssetBuilder.h"
+#include "Build/CameraBuildContext.h"
 #include "Build/CameraRigAssetBuilder.h"
 #include "Build/CameraShakeAssetBuilder.h"
 #include "Core/CameraAsset.h"
@@ -181,10 +182,10 @@ struct TCameraObjectBuilderTraits;
 template<>
 struct TCameraObjectBuilderTraits<UCameraAsset>
 {
-	static void Build(UCameraAsset* InObject, FCameraBuildLog& InBuildLog)
+	static void Build(UCameraAsset* InObject, FCameraBuildContext& InBuildContext)
 	{
 		const bool bBuildReferencedAssets = false; // We are going to build rigs first.
-		FCameraAssetBuilder Builder(InBuildLog);
+		FCameraAssetBuilder Builder(InBuildContext);
 		Builder.BuildCamera(InObject, bBuildReferencedAssets); 
 	}
 };
@@ -192,9 +193,9 @@ struct TCameraObjectBuilderTraits<UCameraAsset>
 template<>
 struct TCameraObjectBuilderTraits<UCameraRigAsset>
 {
-	static void Build(UCameraRigAsset* InObject, FCameraBuildLog& InBuildLog)
+	static void Build(UCameraRigAsset* InObject, FCameraBuildContext& InBuildContext)
 	{
-		FCameraRigAssetBuilder Builder(InBuildLog);
+		FCameraRigAssetBuilder Builder(InBuildContext);
 		Builder.BuildCameraRig(InObject);
 	}
 };
@@ -202,9 +203,9 @@ struct TCameraObjectBuilderTraits<UCameraRigAsset>
 template<>
 struct TCameraObjectBuilderTraits<UCameraShakeAsset>
 {
-	static void Build(UCameraShakeAsset* InObject, FCameraBuildLog& InBuildLog)
+	static void Build(UCameraShakeAsset* InObject, FCameraBuildContext& InBuildContext)
 	{
-		FCameraShakeAssetBuilder Builder(InBuildLog);
+		FCameraShakeAssetBuilder Builder(InBuildContext);
 		Builder.BuildCameraShake(InObject); 
 	}
 };
@@ -235,13 +236,14 @@ struct TCameraObjectBuilderUtil
 	{
 		FCameraBuildLog BuildLog;
 		BuildLog.SetForwardMessagesToLogging(true);
+		FCameraBuildContext BuildContext(BuildLog, ECameraBuildReason::StartingPIE);
 
 		for (CameraObjectType* CameraObject : InCameraObjectsToBuild)
 		{
 			BuildLog.ResetMessages();
 			BuildLog.SetLoggingPrefix(CameraObject->GetName());
 
-			TCameraObjectBuilderTraits<CameraObjectType>::Build(CameraObject, BuildLog);
+			TCameraObjectBuilderTraits<CameraObjectType>::Build(CameraObject, BuildContext);
 		}
 	}
 };
diff --git a/Engine/Plugins/Cameras/GameplayCameras/Source/GameplayCamerasEditor/Private/Toolkits/CameraAssetEditorToolkit.cpp b/Engine/Plugins/Cameras/GameplayCameras/Source/GameplayCamerasEditor/Private/Toolkits/CameraAssetEditorToolkit.cpp
--- a/Engine/Plugins/Cameras/GameplayCameras/Source/GameplayCamerasEditor/Private/Toolkits/CameraAssetEditorToolkit.cpp
+++ b/Engine/Plugins/Cameras/GameplayCameras/Source/GameplayCamerasEditor/Private/Toolkits/CameraAssetEditorToolkit.cpp
@@ -12,6 +12,7 @@
 #include "Core/CameraDirector.h"
 #include "Core/CameraRigAsset.h"
 #include "Editors/ObjectTreeGraphConfig.h"
+#include "Editors/SCameraAssetInterfaceParametersPanel.h"
 #include "Editors/SFindInObjectTreeGraph.h"
 #include "FileHelpers.h"
 #include "Framework/Docking/LayoutExtender.h"
@@ -47,6 +48,7 @@ namespace UE::Cameras
 
 const FName FCameraAssetEditorToolkit::SearchTabId(TEXT("CameraAssetEditor_Search"));
 const FName FCameraAssetEditorToolkit::MessagesTabId(TEXT("CameraAssetEditor_Messages"));
+const FName FCameraAssetEditorToolkit::InterfaceParametersTabId(TEXT("CameraAssetEditor_InterfaceParameters"));
 
 FCameraAssetEditorToolkit::FCameraAssetEditorToolkit(UCameraAssetEditor* InOwningAssetEditor)
 	: FAssetEditorModeManagerToolkit(InOwningAssetEditor)
@@ -59,6 +61,7 @@ FCameraAssetEditorToolkit::FCameraAssetEditorToolkit(UCameraAssetEditor* InOwnin
 
 	StandardLayout->AddBottomTab(SearchTabId);
 	StandardLayout->AddBottomTab(MessagesTabId);
+	StandardLayout->AddLeftTab(InterfaceParametersTabId, ETabState::OpenedTab);
 
 	TSharedPtr<FLayoutExtender> NewLayoutExtender = MakeShared<FLayoutExtender>();
 	{
@@ -71,6 +74,11 @@ FCameraAssetEditorToolkit::FCameraAssetEditorToolkit(UCameraAssetEditor* InOwnin
 				FStandardToolkitLayout::BottomStackExtensionId,
 				ELayoutExtensionPosition::After,
 				FTabManager::FTab(MessagesTabId, ETabState::ClosedTab));
+
+		NewLayoutExtender->ExtendStack(
+				FStandardToolkitLayout::LeftStackExtensionId,
+				ELayoutExtensionPosition::Above,
+				FTabManager::FTab(InterfaceParametersTabId, ETabState::OpenedTab));
 	}
 	LayoutExtenders.Add(NewLayoutExtender);
 }
@@ -105,6 +113,11 @@ void FCameraAssetEditorToolkit::RegisterTabSpawners(const TSharedRef<FTabManager
 		.SetDisplayName(LOCTEXT("Messages", "Messages"))
 		.SetGroup(AssetEditorTabsCategory.ToSharedRef())
 		.SetIcon(FSlateIcon(CamerasStyleSetName, "CameraAssetEditor.Tabs.Messages"));
+
+	InTabManager->RegisterTabSpawner(InterfaceParametersTabId, FOnSpawnTab::CreateSP(this, &FCameraAssetEditorToolkit::SpawnTab_InterfaceParameters))
+		.SetDisplayName(LOCTEXT("InterfaceParameters", "Interface Parameters"))
+		.SetGroup(AssetEditorTabsCategory.ToSharedRef())
+		.SetIcon(FSlateIcon(CamerasStyleSetName, "CameraRigAssetEditor.Tabs.InterfaceParameters"));
 }
 
 TSharedRef<SDockTab> FCameraAssetEditorToolkit::SpawnTab_Search(const FSpawnTabArgs& Args)
@@ -129,6 +142,17 @@ TSharedRef<SDockTab> FCameraAssetEditorToolkit::SpawnTab_Messages(const FSpawnTa
 	return MessagesTab.ToSharedRef();
 }
 
+TSharedRef<SDockTab> FCameraAssetEditorToolkit::SpawnTab_InterfaceParameters(const FSpawnTabArgs& Args)
+{
+	TSharedPtr<SDockTab> InterfaceParametersTab = SNew(SDockTab)
+		.Label(LOCTEXT("InterfaceParametersTabTitle", "Parameters"))
+		[
+			InterfaceParametersPanel.ToSharedRef()
+		];
+
+	return InterfaceParametersTab.ToSharedRef();
+}
+
 void FCameraAssetEditorToolkit::UnregisterTabSpawners(const TSharedRef<FTabManager>& InTabManager)
 {
 	// Skip FBaseAssetToolkit here because we don't want a viewport tab.
@@ -136,6 +160,7 @@ void FCameraAssetEditorToolkit::UnregisterTabSpawners(const TSharedRef<FTabManag
 
 	InTabManager->UnregisterTabSpawner(SearchTabId);
 	InTabManager->UnregisterTabSpawner(MessagesTabId);
+	InTabManager->UnregisterTabSpawner(InterfaceParametersTabId);
 }
 
 void FCameraAssetEditorToolkit::CreateWidgets()
@@ -156,6 +181,10 @@ void FCameraAssetEditorToolkit::CreateWidgets()
 		.OnGetGraphsToSearch(this, &FCameraAssetEditorToolkit::OnGetGraphsToSearch)
 		.OnJumpToNodeRequested(this, &FCameraAssetEditorToolkit::OnJumpToNode);
 
+	// Create the interface parameters panel.
+	InterfaceParametersPanel = SNew(SCameraAssetInterfaceParametersPanel, this);
+	InterfaceParametersPanel->OnInterfaceParameterSelected().AddSP(this, &FCameraAssetEditorToolkit::OnInterfaceParameterSelected);
+
 	// Create the message log.
 	BuildLogToolkit->Initialize("CameraAssetBuildMessages");
 }
@@ -274,6 +303,11 @@ void FCameraAssetEditorToolkit::PostInitAssetEditor()
 	SetEditorMode(InitialModeName);
 
 	UpgradeLegacyCameraAssets();
+
+	// If the asset predates the FCameraAssetInterface, upgrade it. We don't mark it as modified in this case.
+	// We either let the user modify the asset and save it, or we'll re-upgrade it next time if needed again.
+	// This prevents prompting the user to save changes they don't know about.
+	FCameraAssetBuilder::BuildDefaultInterfaceIfNeeded(CameraAsset);
 }
 
 TSharedPtr<FAssetEditorMode> FCameraAssetEditorToolkit::CreateCameraDirectorAssetEditorMode()
@@ -353,10 +387,9 @@ void FCameraAssetEditorToolkit::OnBuild()
 		return;
 	}
 
-	FCameraDirectorRigUsageInfo UsageInfo;
-
 	FCameraBuildLog BuildLog;
-	FCameraAssetBuilder Builder(BuildLog);
+	FCameraBuildContext BuildContext(BuildLog, ECameraBuildReason::UserAction);
+	FCameraAssetBuilder Builder(BuildContext);
 	Builder.BuildCamera(CameraAsset);
 	
 	BuildLogToolkit->PopulateMessageListing(BuildLog);
@@ -368,14 +401,18 @@ void FCameraAssetEditorToolkit::OnBuild()
 
 	DetailsView->RequestForceRefresh();
 
-	for (UCameraRigAsset* CameraRigAsset : UsageInfo.CameraRigs)
+	TConstArrayView<UObject*> BuiltObjects = Builder.GetBuiltReferencedAssets();
+	for (UObject* BuiltObject : BuiltObjects)
 	{
-		FCameraRigPackages BuiltPackages;
-		CameraRigAsset->GatherPackages(BuiltPackages);
-
-		for (const UPackage* BuiltPackage : BuiltPackages)
+		if (UCameraRigAsset* CameraRigAsset = Cast<UCameraRigAsset>(BuiltObject))
 		{
-			LiveEditManager->NotifyPostBuildAsset(BuiltPackage);
+			FCameraRigPackages BuiltPackages;
+			CameraRigAsset->GatherPackages(BuiltPackages);
+
+			for (const UPackage* BuiltPackage : BuiltPackages)
+			{
+				LiveEditManager->NotifyPostBuildAsset(BuiltPackage);
+			}
 		}
 	}
 
@@ -458,6 +495,11 @@ void FCameraAssetEditorToolkit::OnJumpToObject(UObject* Object, FName PropertyNa
 	}
 }
 
+void FCameraAssetEditorToolkit::OnInterfaceParameterSelected(UCameraAssetInterfaceParameter* InterfaceParameter)
+{
+	DetailsView->SetObject(InterfaceParameter);
+}
+
 FText FCameraAssetEditorToolkit::GetBaseToolkitName() const
 {
 	return LOCTEXT("AppLabel", "Camera Asset");
diff --git a/Engine/Plugins/Cameras/GameplayCameras/Source/GameplayCamerasEditor/Private/Toolkits/CameraRigAssetEditorToolkit.cpp b/Engine/Plugins/Cameras/GameplayCameras/Source/GameplayCamerasEditor/Private/Toolkits/CameraRigAssetEditorToolkit.cpp
--- a/Engine/Plugins/Cameras/GameplayCameras/Source/GameplayCamerasEditor/Private/Toolkits/CameraRigAssetEditorToolkit.cpp
+++ b/Engine/Plugins/Cameras/GameplayCameras/Source/GameplayCamerasEditor/Private/Toolkits/CameraRigAssetEditorToolkit.cpp
@@ -115,7 +115,7 @@ void FCameraRigAssetEditorToolkit::RegisterTabSpawners(const TSharedRef<class FT
 		.SetIcon(FSlateIcon(CamerasStyleSetName, "CameraRigAssetEditor.Tabs.Curves"));
 
 	InTabManager->RegisterTabSpawner(InterfaceParametersTabId, FOnSpawnTab::CreateSP(this, &FCameraRigAssetEditorToolkit::SpawnTab_InterfaceParameters))
-		.SetDisplayName(LOCTEXT("InterfaceParameters", "InterfaceParameters"))
+		.SetDisplayName(LOCTEXT("InterfaceParameters", "Interface Parameters"))
 		.SetGroup(AssetEditorTabsCategory.ToSharedRef())
 		.SetIcon(FSlateIcon(CamerasStyleSetName, "CameraRigAssetEditor.Tabs.InterfaceParameters"));
 }
@@ -353,7 +353,8 @@ void FCameraRigAssetEditorToolkit::OnBuild()
 	}
 
 	FCameraBuildLog BuildLog;
-	FCameraRigAssetBuilder Builder(BuildLog);
+	FCameraBuildContext BuildContext(BuildLog, ECameraBuildReason::UserAction);
+	FCameraRigAssetBuilder Builder(BuildContext);
 	Builder.BuildCameraRig(CameraRigAsset);
 
 	BuildLogToolkit->PopulateMessageListing(BuildLog);
diff --git a/Engine/Plugins/Cameras/GameplayCameras/Source/GameplayCamerasEditor/Private/Toolkits/CameraShakeAssetEditorToolkit.cpp b/Engine/Plugins/Cameras/GameplayCameras/Source/GameplayCamerasEditor/Private/Toolkits/CameraShakeAssetEditorToolkit.cpp
--- a/Engine/Plugins/Cameras/GameplayCameras/Source/GameplayCamerasEditor/Private/Toolkits/CameraShakeAssetEditorToolkit.cpp
+++ b/Engine/Plugins/Cameras/GameplayCameras/Source/GameplayCamerasEditor/Private/Toolkits/CameraShakeAssetEditorToolkit.cpp
@@ -385,7 +385,8 @@ void FCameraShakeAssetEditorToolkit::OnBuild()
 	}
 
 	FCameraBuildLog BuildLog;
-	FCameraShakeAssetBuilder Builder(BuildLog);
+	FCameraBuildContext BuildContext(BuildLog, ECameraBuildReason::UserAction);
+	FCameraShakeAssetBuilder Builder(BuildContext);
 	Builder.BuildCameraShake(CameraShakeAsset);
 
 	BuildLogToolkit->PopulateMessageListing(BuildLog);
diff --git a/Engine/Plugins/Cameras/GameplayCameras/Source/GameplayCamerasEditor/Public/Editors/SCameraAssetInterfaceParametersPanel.h b/Engine/Plugins/Cameras/GameplayCameras/Source/GameplayCamerasEditor/Public/Editors/SCameraAssetInterfaceParametersPanel.h
new file mode 100644
--- /dev/null
+++ b/Engine/Plugins/Cameras/GameplayCameras/Source/GameplayCamerasEditor/Public/Editors/SCameraAssetInterfaceParametersPanel.h
@@ -0,0 +1,79 @@
+// Copyright Epic Games, Inc. All Rights Reserved.
+
+#pragma once
+
+#include "Core/CameraAsset.h"
+#include "Templates/SharedPointerFwd.h"
+#include "Widgets/SCompoundWidget.h"
+
+class UCameraAssetInterfaceParameter;
+class ITableRow;
+class STableViewBase;
+template <typename ItemType> class SListView;
+
+namespace UE::Cameras
+{
+
+class FCameraAssetEditorToolkit;
+
+DECLARE_MULTICAST_DELEGATE_OneParam(FOnCameraAssetInterfaceParameterEvent, UCameraAssetInterfaceParameter*);
+
+class SCameraAssetInterfaceParametersPanel 
+	: public SCompoundWidget
+	, public ICameraAssetEventHandler
+{
+public:
+
+	static const FName ParameterTypeColumn;
+	static const FName ParameterNameColumn;
+	static const FName SourceCameraRigColumn;
+	static const FName SourceParameterNameColumn;
+
+public:
+
+	SLATE_BEGIN_ARGS(SCameraAssetInterfaceParametersPanel)
+	{}
+	SLATE_END_ARGS()
+
+	void Construct(const FArguments& Args, FCameraAssetEditorToolkit* OwnerToolkit);
+
+	/** Request that the list of parameters be refreshed. */
+	void RequestListRefresh();
+
+	/** Delegate invoked when a parmeter is selected in the panel. */
+	FOnCameraAssetInterfaceParameterEvent& OnInterfaceParameterSelected() { return OnInterfaceParameterSelectedDelegate; }
+
+protected:
+
+	// SWidget interface.
+	virtual void Tick(const FGeometry& AllottedGeometry, const double InCurrentTime, const float InDeltaTime) override;
+
+	// ICameraAssetEventHandler interface.
+	virtual void OnCameraAssetInterfaceChanged() override;
+
+private:
+
+	TSharedRef<ITableRow> OnGenerateInterfaceParameterRow(TObjectPtr<UCameraAssetInterfaceParameter> Item, const TSharedRef<STableViewBase>& OwnerTable);
+	void OnInterfaceParameterSelectionChanged(TObjectPtr<UCameraAssetInterfaceParameter> Item, ESelectInfo::Type Type);
+	TSharedPtr<SWidget> OnInterfaceParameterContextMenuOpening();
+
+	FReply OnAddInterfaceParameter();
+	void OnRenameInterfaceParameter(TObjectPtr<UCameraAssetInterfaceParameter> Item);
+	void OnDeleteInterfaceParameter(TObjectPtr<UCameraAssetInterfaceParameter> Item);
+
+private:
+
+	UCameraAsset* CameraAsset = nullptr;
+	FCameraAssetEditorToolkit* Toolkit = nullptr;
+
+	TCameraEventHandler<ICameraAssetEventHandler> EventHandler;
+
+	TSharedPtr<SListView<TObjectPtr<UCameraAssetInterfaceParameter>>> ParametersListView;
+
+	FOnCameraAssetInterfaceParameterEvent OnInterfaceParameterSelectedDelegate;
+
+	bool bListRefreshRequested = false;
+};
+
+}  // namespace UE::Cameras
+
diff --git a/Engine/Plugins/Cameras/GameplayCameras/Source/GameplayCamerasEditor/Public/Editors/SCameraObjectInterfaceParametersPanel.h b/Engine/Plugins/Cameras/GameplayCameras/Source/GameplayCamerasEditor/Public/Editors/SCameraObjectInterfaceParametersPanel.h
--- a/Engine/Plugins/Cameras/GameplayCameras/Source/GameplayCamerasEditor/Public/Editors/SCameraObjectInterfaceParametersPanel.h
+++ b/Engine/Plugins/Cameras/GameplayCameras/Source/GameplayCamerasEditor/Public/Editors/SCameraObjectInterfaceParametersPanel.h
@@ -27,6 +27,11 @@ class SCameraObjectInterfaceParametersPanel
 	: public SCompoundWidget
 	, public ICameraObjectEventHandler
 {
+public:
+
+	static const FName ParameterTypeColumn;
+	static const FName ParameterNameColumn;
+
 public:
 
 	SLATE_BEGIN_ARGS(SCameraObjectInterfaceParametersPanel)
diff --git a/Engine/Plugins/Cameras/GameplayCameras/Source/GameplayCamerasEditor/Public/Toolkits/CameraAssetEditorToolkit.h b/Engine/Plugins/Cameras/GameplayCameras/Source/GameplayCamerasEditor/Public/Toolkits/CameraAssetEditorToolkit.h
--- a/Engine/Plugins/Cameras/GameplayCameras/Source/GameplayCamerasEditor/Public/Toolkits/CameraAssetEditorToolkit.h
+++ b/Engine/Plugins/Cameras/GameplayCameras/Source/GameplayCamerasEditor/Public/Toolkits/CameraAssetEditorToolkit.h
@@ -24,6 +24,7 @@ class FBuildButtonToolkit;
 class FCameraBuildLogToolkit;
 class FStandardToolkitLayout;
 class IGameplayCamerasLiveEditManager;
+class SCameraAssetInterfaceParametersPanel;
 
 /**
  * Editor toolkit for a camera asset.
@@ -38,6 +39,8 @@ public:
 	FCameraAssetEditorToolkit(UCameraAssetEditor* InOwningAssetEditor);
 	~FCameraAssetEditorToolkit();
 
+	UCameraAsset* GetCameraAsset() const { return CameraAsset; }
+
 protected:
 
 	// FBaseAssetToolkit interface
@@ -69,6 +72,7 @@ private:
 
 	TSharedRef<SDockTab> SpawnTab_Search(const FSpawnTabArgs& Args);
 	TSharedRef<SDockTab> SpawnTab_Messages(const FSpawnTabArgs& Args);
+	TSharedRef<SDockTab> SpawnTab_InterfaceParameters(const FSpawnTabArgs& Args);
 
 	void FillCameraMenu(FMenuBuilder& MenuBuilder);
 
@@ -84,12 +88,15 @@ private:
 	void OnJumpToObject(UObject* Object);
 	void OnJumpToObject(UObject* Object, FName PropertyName);
 
+	void OnInterfaceParameterSelected(UCameraAssetInterfaceParameter* InterfaceParameter);
+
 	void UpgradeLegacyCameraAssets();
 
 private:
 
 	static const FName SearchTabId;
 	static const FName MessagesTabId;
+	static const FName InterfaceParametersTabId;
 
 	/** The asset being edited */
 	TObjectPtr<UCameraAsset> CameraAsset;
@@ -109,6 +116,10 @@ private:
 	/** The search results window */
 	TSharedPtr<SFindInObjectTreeGraph> SearchWidget;
 
+	/** The panel for exposing parameters. */
+	TSharedPtr<SCameraAssetInterfaceParametersPanel> InterfaceParametersPanel;
+
+	/** A cached pointer to the editor settings. */
 	TObjectPtr<UGameplayCamerasEditorSettings> Settings;
 
 	/** Live edit manager for updating the assets in the runtime */
diff --git a/Engine/Plugins/Cameras/GameplayCameras/Source/GameplayCamerasUncookedOnly/Private/BlueprintGraph/K2Node_CameraRigBase.cpp b/Engine/Plugins/Cameras/GameplayCameras/Source/GameplayCamerasUncookedOnly/Private/BlueprintGraph/K2Node_CameraRigBase.cpp
--- a/Engine/Plugins/Cameras/GameplayCameras/Source/GameplayCamerasUncookedOnly/Private/BlueprintGraph/K2Node_CameraRigBase.cpp
+++ b/Engine/Plugins/Cameras/GameplayCameras/Source/GameplayCamerasUncookedOnly/Private/BlueprintGraph/K2Node_CameraRigBase.cpp
@@ -12,6 +12,8 @@
 #include "GameFramework/BlueprintCameraEvaluationDataRef.h"
 #include "GameFramework/CameraRigParameterInterop.h"
 #include "GameplayCamerasDelegates.h"
+#include "Helpers/CameraContextDataPinTypeHelper.h"
+#include "Helpers/CameraVariablePinTypeHelper.h"
 #include "K2Node_CallFunction.h"
 #include "K2Node_EnumLiteral.h"
 #include "Kismet/KismetSystemLibrary.h"
@@ -145,57 +147,7 @@ FEdGraphPinType UK2Node_CameraRigBase::MakeBlendableParameterPinType(const UCame
 
 FEdGraphPinType UK2Node_CameraRigBase::MakeBlendableParameterPinType(ECameraVariableType CameraVariableType, const UScriptStruct* BlendableStructType)
 {
-	FName PinCategory;
-	FName PinSubCategory;
-	UObject* PinSubCategoryObject = nullptr;
-	switch (CameraVariableType)
-	{
-		case ECameraVariableType::Boolean:
-			PinCategory = UEdGraphSchema_K2::PC_Boolean;
-			break;
-		case ECameraVariableType::Integer32:
-			PinCategory = UEdGraphSchema_K2::PC_Int;
-			break;
-		case ECameraVariableType::Float:
-			// We'll cast down to float.
-			PinCategory = UEdGraphSchema_K2::PC_Real;
-			PinSubCategory = UEdGraphSchema_K2::PC_Float;
-			break;
-		case ECameraVariableType::Double:
-			PinCategory = UEdGraphSchema_K2::PC_Real;
-			PinSubCategory = UEdGraphSchema_K2::PC_Double;
-			break;
-		case ECameraVariableType::Vector2d:
-			PinCategory = UEdGraphSchema_K2::PC_Struct;
-			PinSubCategoryObject = TBaseStructure<FVector2D>::Get();
-			break;
-		case ECameraVariableType::Vector3d:
-			PinCategory = UEdGraphSchema_K2::PC_Struct;
-			PinSubCategoryObject = TBaseStructure<FVector>::Get();
-			break;
-		case ECameraVariableType::Vector4d:
-			PinCategory = UEdGraphSchema_K2::PC_Struct;
-			PinSubCategoryObject = TBaseStructure<FVector4>::Get();
-			break;
-		case ECameraVariableType::Rotator3d:
-			PinCategory = UEdGraphSchema_K2::PC_Struct;
-			PinSubCategoryObject = TBaseStructure<FRotator>::Get();
-			break;
-		case ECameraVariableType::Transform3d:
-			PinCategory = UEdGraphSchema_K2::PC_Struct;
-			PinSubCategoryObject = TBaseStructure<FTransform>::Get();
-			break;
-		case ECameraVariableType::BlendableStruct:
-			PinCategory = UEdGraphSchema_K2::PC_Struct;
-			PinSubCategoryObject = const_cast<UScriptStruct*>(BlendableStructType);
-			break;
-	}
-
-	FEdGraphPinType PinType;
-	PinType.PinCategory = PinCategory;
-	PinType.PinSubCategory = PinSubCategory;
-	PinType.PinSubCategoryObject = PinSubCategoryObject;
-	return PinType;
+	return UE::Cameras::FCameraVariablePinTypeHelper::GetPinType(CameraVariableType, BlendableStructType);
 }
 
 FEdGraphPinType UK2Node_CameraRigBase::MakeDataParameterPinType(const UCameraObjectInterfaceDataParameter* DataParameter)
@@ -205,44 +157,7 @@ FEdGraphPinType UK2Node_CameraRigBase::MakeDataParameterPinType(const UCameraObj
 
 FEdGraphPinType UK2Node_CameraRigBase::MakeDataParameterPinType(ECameraContextDataType CameraContextDataType, ECameraContextDataContainerType CameraContextDataContainerType, const UObject* CameraContextDataTypeObject)
 {
-	FName PinCategory;
-	UObject* PinSubCategoryObject = const_cast<UObject*>(CameraContextDataTypeObject);
-	switch (CameraContextDataType)
-	{
-		case ECameraContextDataType::Name:
-			PinCategory = UEdGraphSchema_K2::PC_Name;
-			break;
-		case ECameraContextDataType::String:
-			PinCategory = UEdGraphSchema_K2::PC_String;
-			break;
-		case ECameraContextDataType::Enum:
-			PinCategory = UEdGraphSchema_K2::PC_Byte;
-			break;
-		case ECameraContextDataType::Struct:
-			PinCategory = UEdGraphSchema_K2::PC_Struct;
-			break;
-		case ECameraContextDataType::Object:
-			PinCategory = UEdGraphSchema_K2::PC_Object;
-			break;
-		case ECameraContextDataType::Class:
-			PinCategory = UEdGraphSchema_K2::PC_Class;
-			break;
-	}
-
-	EPinContainerType PinContainerType = EPinContainerType::None;
-	switch (CameraContextDataContainerType)
-	{
-		case ECameraContextDataContainerType::Array:
-			PinContainerType = EPinContainerType::Array;
-		default:
-			break;
-	}
-
-	FEdGraphPinType PinType;
-	PinType.PinCategory = PinCategory;
-	PinType.PinSubCategoryObject = PinSubCategoryObject;
-	PinType.ContainerType = PinContainerType;
-	return PinType;
+	return UE::Cameras::FCameraContextDataPinTypeHelper::GetPinType(CameraContextDataType, CameraContextDataContainerType, CameraContextDataTypeObject);
 }
 
 UK2Node_CallFunction* UK2Node_CameraRigBase::CreateMakeLiteralNode(FKismetCompilerContext& CompilerContext, UEdGraph* SourceGraph, UK2Node* SourceNode, UClass* FunctionLibraryClass, const TCHAR* FunctionName, UEdGraphPin* SourceValuePin)
diff --git a/Engine/Plugins/Cameras/GameplayCameras/Source/GameplayCamerasUncookedOnly/Private/Helpers/CameraContextDataPinTypeHelper.cpp b/Engine/Plugins/Cameras/GameplayCameras/Source/GameplayCamerasUncookedOnly/Private/Helpers/CameraContextDataPinTypeHelper.cpp
new file mode 100644
--- /dev/null
+++ b/Engine/Plugins/Cameras/GameplayCameras/Source/GameplayCamerasUncookedOnly/Private/Helpers/CameraContextDataPinTypeHelper.cpp
@@ -0,0 +1,55 @@
+// Copyright Epic Games, Inc. All Rights Reserved.
+
+#pragma once
+
+#include "Helpers/CameraContextDataPinTypeHelper.h"
+
+#include "EdGraphSchema_K2.h"
+
+namespace UE::Cameras
+{
+
+FEdGraphPinType FCameraContextDataPinTypeHelper::GetPinType(ECameraContextDataType DataType, ECameraContextDataContainerType DataContainerType, const UObject* DataTypeObject)
+{
+	FName PinCategory = UEdGraphSchema_K2::PC_Name;
+	UObject* PinSubCategoryObject = const_cast<UObject*>(DataTypeObject);
+	switch (DataType)
+	{
+		case ECameraContextDataType::Name:
+			PinCategory = UEdGraphSchema_K2::PC_Name;
+			break;
+		case ECameraContextDataType::String:
+			PinCategory = UEdGraphSchema_K2::PC_String;
+			break;
+		case ECameraContextDataType::Enum:
+			PinCategory = UEdGraphSchema_K2::PC_Byte;
+			break;
+		case ECameraContextDataType::Struct:
+			PinCategory = UEdGraphSchema_K2::PC_Struct;
+			break;
+		case ECameraContextDataType::Object:
+			PinCategory = UEdGraphSchema_K2::PC_Object;
+			break;
+		case ECameraContextDataType::Class:
+			PinCategory = UEdGraphSchema_K2::PC_Class;
+			break;
+	}
+
+	EPinContainerType PinContainerType = EPinContainerType::None;
+	switch (DataContainerType)
+	{
+		case ECameraContextDataContainerType::Array:
+			PinContainerType = EPinContainerType::Array;
+		default:
+			break;
+	}
+
+	FEdGraphPinType PinType;
+	PinType.PinCategory = PinCategory;
+	PinType.PinSubCategoryObject = PinSubCategoryObject;
+	PinType.ContainerType = PinContainerType;
+	return PinType;
+}
+
+}  // namespace UE::Cameras
+
diff --git a/Engine/Plugins/Cameras/GameplayCameras/Source/GameplayCamerasUncookedOnly/Private/Helpers/CameraVariablePinTypeHelper.cpp b/Engine/Plugins/Cameras/GameplayCameras/Source/GameplayCamerasUncookedOnly/Private/Helpers/CameraVariablePinTypeHelper.cpp
new file mode 100644
--- /dev/null
+++ b/Engine/Plugins/Cameras/GameplayCameras/Source/GameplayCamerasUncookedOnly/Private/Helpers/CameraVariablePinTypeHelper.cpp
@@ -0,0 +1,88 @@
+// Copyright Epic Games, Inc. All Rights Reserved.
+
+#pragma once
+
+#include "Helpers/CameraVariablePinTypeHelper.h"
+
+#include "EdGraphSchema_K2.h"
+
+namespace UE::Cameras
+{
+
+FEdGraphPinType FCameraVariablePinTypeHelper::GetPinType(ECameraVariableType VariableType, const UScriptStruct* BlendableStructType)
+{
+	FName PinCategory;
+	FName PinSubCategory;
+	UObject* PinSubCategoryObject = nullptr;
+	switch (VariableType)
+	{
+		case ECameraVariableType::Boolean:
+			PinCategory = UEdGraphSchema_K2::PC_Boolean;
+			break;
+		case ECameraVariableType::Integer32:
+			PinCategory = UEdGraphSchema_K2::PC_Int;
+			break;
+		case ECameraVariableType::Float:
+			// We'll cast down to float.
+			PinCategory = UEdGraphSchema_K2::PC_Real;
+			PinSubCategory = UEdGraphSchema_K2::PC_Float;
+			break;
+		case ECameraVariableType::Double:
+			PinCategory = UEdGraphSchema_K2::PC_Real;
+			PinSubCategory = UEdGraphSchema_K2::PC_Double;
+			break;
+		case ECameraVariableType::Vector2f:
+			PinCategory = UEdGraphSchema_K2::PC_Struct;
+			PinSubCategoryObject = TVariantStructure<FVector2f>::Get();
+			break;
+		case ECameraVariableType::Vector2d:
+			PinCategory = UEdGraphSchema_K2::PC_Struct;
+			PinSubCategoryObject = TBaseStructure<FVector2D>::Get();
+			break;
+		case ECameraVariableType::Vector3f:
+			PinCategory = UEdGraphSchema_K2::PC_Struct;
+			PinSubCategoryObject = TVariantStructure<FVector3f>::Get();
+			break;
+		case ECameraVariableType::Vector3d:
+			PinCategory = UEdGraphSchema_K2::PC_Struct;
+			PinSubCategoryObject = TBaseStructure<FVector>::Get();
+			break;
+		case ECameraVariableType::Vector4f:
+			PinCategory = UEdGraphSchema_K2::PC_Struct;
+			PinSubCategoryObject = TVariantStructure<FVector4f>::Get();
+			break;
+		case ECameraVariableType::Vector4d:
+			PinCategory = UEdGraphSchema_K2::PC_Struct;
+			PinSubCategoryObject = TBaseStructure<FVector4>::Get();
+			break;
+		case ECameraVariableType::Rotator3f:
+			PinCategory = UEdGraphSchema_K2::PC_Struct;
+			PinSubCategoryObject = TVariantStructure<FRotator3f>::Get();
+			break;
+		case ECameraVariableType::Rotator3d:
+			PinCategory = UEdGraphSchema_K2::PC_Struct;
+			PinSubCategoryObject = TBaseStructure<FRotator>::Get();
+			break;
+		case ECameraVariableType::Transform3f:
+			PinCategory = UEdGraphSchema_K2::PC_Struct;
+			PinSubCategoryObject = TVariantStructure<FTransform3f>::Get();
+			break;
+		case ECameraVariableType::Transform3d:
+			PinCategory = UEdGraphSchema_K2::PC_Struct;
+			PinSubCategoryObject = TBaseStructure<FTransform>::Get();
+			break;
+		case ECameraVariableType::BlendableStruct:
+			PinCategory = UEdGraphSchema_K2::PC_Struct;
+			PinSubCategoryObject = const_cast<UScriptStruct*>(BlendableStructType);
+			break;
+	}
+
+	FEdGraphPinType PinType;
+	PinType.PinCategory = PinCategory;
+	PinType.PinSubCategory = PinSubCategory;
+	PinType.PinSubCategoryObject = PinSubCategoryObject;
+	return PinType;
+}
+
+}  // namespace UE::Cameras
+
diff --git a/Engine/Plugins/Cameras/GameplayCameras/Source/GameplayCamerasUncookedOnly/Public/Helpers/CameraContextDataPinTypeHelper.h b/Engine/Plugins/Cameras/GameplayCameras/Source/GameplayCamerasUncookedOnly/Public/Helpers/CameraContextDataPinTypeHelper.h
new file mode 100644
--- /dev/null
+++ b/Engine/Plugins/Cameras/GameplayCameras/Source/GameplayCamerasUncookedOnly/Public/Helpers/CameraContextDataPinTypeHelper.h
@@ -0,0 +1,23 @@
+// Copyright Epic Games, Inc. All Rights Reserved.
+
+#pragma once
+
+#include "Core/CameraContextDataTable.h"
+#include "EdGraph/EdGraphPin.h"
+
+#define UE_API GAMEPLAYCAMERASUNCOOKEDONLY_API
+
+namespace UE::Cameras
+{
+
+class FCameraContextDataPinTypeHelper
+{
+public:
+
+	UE_API static FEdGraphPinType GetPinType(ECameraContextDataType DataType, ECameraContextDataContainerType DataContainerType, const UObject* DataTypeObject);
+};
+
+}  // namespace UE::Cameras
+
+#undef UE_API
+
diff --git a/Engine/Plugins/Cameras/GameplayCameras/Source/GameplayCamerasUncookedOnly/Public/Helpers/CameraVariablePinTypeHelper.h b/Engine/Plugins/Cameras/GameplayCameras/Source/GameplayCamerasUncookedOnly/Public/Helpers/CameraVariablePinTypeHelper.h
new file mode 100644
--- /dev/null
+++ b/Engine/Plugins/Cameras/GameplayCameras/Source/GameplayCamerasUncookedOnly/Public/Helpers/CameraVariablePinTypeHelper.h
@@ -0,0 +1,23 @@
+// Copyright Epic Games, Inc. All Rights Reserved.
+
+#pragma once
+
+#include "Core/CameraVariableTable.h"
+#include "EdGraph/EdGraphPin.h"
+
+#define UE_API GAMEPLAYCAMERASUNCOOKEDONLY_API
+
+namespace UE::Cameras
+{
+
+class FCameraVariablePinTypeHelper
+{
+public:
+
+	UE_API static FEdGraphPinType GetPinType(ECameraVariableType VariableType, const UScriptStruct* BlendableStructType);
+};
+
+}  // namespace UE::Cameras
+
+#undef UE_API
+
