Cameras: fix unity builds potentially leading to symbol redefinition

#Synced-CL 49781371

diff --git a/Engine/Plugins/Cameras/GameplayCameras/Source/GameplayCameras/Private/Core/CameraContextDataTable.cpp b/Engine/Plugins/Cameras/GameplayCameras/Source/GameplayCameras/Private/Core/CameraContextDataTable.cpp
--- a/Engine/Plugins/Cameras/GameplayCameras/Source/GameplayCameras/Private/Core/CameraContextDataTable.cpp
+++ b/Engine/Plugins/Cameras/GameplayCameras/Source/GameplayCameras/Private/Core/CameraContextDataTable.cpp
@@ -7,13 +7,8 @@
 namespace UE::Cameras
 {
 
-namespace Private
-{
-
-static const uint32 GDefaultCapacity = 64;
-static const uint32 GDefaultAlignment = alignof(void*);
-
-}  // namespace Private
+const uint32 FCameraContextDataTable::GDefaultCapacity = 64;
+const uint32 FCameraContextDataTable::GDefaultAlignment = alignof(void*);
 
 FCameraContextDataTable::FCameraContextDataTable()
 {
@@ -126,7 +121,7 @@ void FCameraContextDataTable::Initialize(const FCameraContextDataTableAllocation
 	}
 
 	// Allocate the memory buffer.
-	MaxAlignOf = FMath::Max(Private::GDefaultAlignment, MaxAlignOf);
+	MaxAlignOf = FMath::Max(GDefaultAlignment, MaxAlignOf);
 	Memory = reinterpret_cast<uint8*>(FMemory::Malloc(TotalSizeOf, MaxAlignOf));
 	Capacity = TotalSizeOf;
 	Alignment = MaxAlignOf;
@@ -209,7 +204,7 @@ void FCameraContextDataTable::AddData(const FCameraContextDataDefinition& DataDe
 	if (NewUsed > Capacity)
 	{
 		// Alignment may not have been initialized yet.
-		Alignment = FMath::Max(AlignOf, Private::GDefaultAlignment);
+		Alignment = FMath::Max(AlignOf, GDefaultAlignment);
 
 		ReallocateBuffer(NewUsed);
 
@@ -410,7 +405,7 @@ bool FCameraContextDataTable::DestroyDataValue(ECameraContextDataType DataType,
 
 void FCameraContextDataTable::ReallocateBuffer(uint32 MinRequired)
 {
-	uint32 NewCapacity = Capacity <= 0 ? Private::GDefaultCapacity : Capacity * 2;
+	uint32 NewCapacity = Capacity <= 0 ? GDefaultCapacity : Capacity * 2;
 	if (MinRequired > 0)
 	{
 		NewCapacity = FMath::Max(NewCapacity, MinRequired);
diff --git a/Engine/Plugins/Cameras/GameplayCameras/Source/GameplayCameras/Private/Core/CameraVariableTable.cpp b/Engine/Plugins/Cameras/GameplayCameras/Source/GameplayCameras/Private/Core/CameraVariableTable.cpp
--- a/Engine/Plugins/Cameras/GameplayCameras/Source/GameplayCameras/Private/Core/CameraVariableTable.cpp
+++ b/Engine/Plugins/Cameras/GameplayCameras/Source/GameplayCameras/Private/Core/CameraVariableTable.cpp
@@ -12,9 +12,6 @@ namespace UE::Cameras
 namespace Private
 {
 
-static const uint32 GDefaultCapacity = 64;
-static const uint32 GDefaultAlignment = alignof(void*);
-
 static const FString GUnavailableVariableDebugName(TEXT("<no debug info>"));
 
 bool IsVariableInMask(FCameraVariableID VariableID, const FCameraVariableTableFlags* InMask, bool bInvertMask)
@@ -29,6 +26,9 @@ bool IsVariableInMask(FCameraVariableID VariableID, const FCameraVariableTableFl
 
 }  // namespace Private
 
+const uint32 FCameraVariableTable::GDefaultCapacity = 64;
+const uint32 FCameraVariableTable::GDefaultAlignment = alignof(void*);
+
 TArray<FBlendableStructInfo> FCameraVariableTable::CachedBlendableStructs;
 bool FCameraVariableTable::bCachedBlendableStructs(false);
 
@@ -107,7 +107,7 @@ void FCameraVariableTable::Initialize(const FCameraVariableTableAllocationInfo&
 	}
 
 	// Allocate the memory buffer.
-	MaxAlignOf = FMath::Max(Private::GDefaultAlignment, MaxAlignOf);
+	MaxAlignOf = FMath::Max(GDefaultAlignment, MaxAlignOf);
 	Memory = reinterpret_cast<uint8*>(FMemory::Malloc(TotalSizeOf, MaxAlignOf));
 	Capacity = TotalSizeOf;
 	Alignment = MaxAlignOf;
@@ -254,7 +254,7 @@ void FCameraVariableTable::AddVariable(const FCameraVariableDefinition& Variable
 	if (NewUsed > Capacity)
 	{
 		// Alignment may not have been initialized yet.
-		Alignment = FMath::Max(AlignOf, Private::GDefaultAlignment);
+		Alignment = FMath::Max(AlignOf, GDefaultAlignment);
 
 		ReallocateBuffer(NewUsed);
 
@@ -274,7 +274,7 @@ void FCameraVariableTable::AddVariable(const FCameraVariableDefinition& Variable
 
 void FCameraVariableTable::ReallocateBuffer(uint32 MinRequired)
 {
-	uint32 NewCapacity = Capacity <= 0 ? Private::GDefaultCapacity : Capacity * 2;
+	uint32 NewCapacity = Capacity <= 0 ? GDefaultCapacity : Capacity * 2;
 	if (MinRequired > 0)
 	{
 		NewCapacity = FMath::Max(NewCapacity, MinRequired);
diff --git a/Engine/Plugins/Cameras/GameplayCameras/Source/GameplayCameras/Public/Core/CameraContextDataTable.h b/Engine/Plugins/Cameras/GameplayCameras/Source/GameplayCameras/Public/Core/CameraContextDataTable.h
--- a/Engine/Plugins/Cameras/GameplayCameras/Source/GameplayCameras/Public/Core/CameraContextDataTable.h
+++ b/Engine/Plugins/Cameras/GameplayCameras/Source/GameplayCameras/Public/Core/CameraContextDataTable.h
@@ -299,6 +299,9 @@ private:
 
 private:
 
+	static const uint32 GDefaultCapacity;
+	static const uint32 GDefaultAlignment;
+
 	TArray<FEntry> Entries;
 	TMap<FCameraContextDataID, int32> EntryLookup;
 
diff --git a/Engine/Plugins/Cameras/GameplayCameras/Source/GameplayCameras/Public/Core/CameraVariableTable.h b/Engine/Plugins/Cameras/GameplayCameras/Source/GameplayCameras/Public/Core/CameraVariableTable.h
--- a/Engine/Plugins/Cameras/GameplayCameras/Source/GameplayCameras/Public/Core/CameraVariableTable.h
+++ b/Engine/Plugins/Cameras/GameplayCameras/Source/GameplayCameras/Public/Core/CameraVariableTable.h
@@ -232,6 +232,9 @@ private:
 
 private:
 
+	static const uint32 GDefaultCapacity;
+	static const uint32 GDefaultAlignment;
+
 	enum class EEntryFlags : uint8
 	{
 		None = 0,
