Cameras: fixes and improvements to the interface parameter panels

- Better default names for newly added parameters.
- Support keyboard shortcuts for renaming and deleting parameters.
- Add a button to delete selected parameters.
- Properly refresh the list of parameters when undoing/redoing an add or delete.

#Synced-CL 49781371

diff --git a/Engine/Plugins/Cameras/GameplayCameras/Source/GameplayCamerasEditor/Private/Commands/CameraAssetEditorCommands.cpp b/Engine/Plugins/Cameras/GameplayCameras/Source/GameplayCamerasEditor/Private/Commands/CameraAssetEditorCommands.cpp
--- a/Engine/Plugins/Cameras/GameplayCameras/Source/GameplayCamerasEditor/Private/Commands/CameraAssetEditorCommands.cpp
+++ b/Engine/Plugins/Cameras/GameplayCameras/Source/GameplayCamerasEditor/Private/Commands/CameraAssetEditorCommands.cpp
@@ -42,6 +42,11 @@ void FCameraAssetEditorCommands::RegisterCommands()
 			EUserInterfaceActionType::ToggleButton, FInputChord());
 	UI_COMMAND(FindInCamera, "Search", "Searches for nodes in this camera asset",
 			EUserInterfaceActionType::Button, FInputChord());
+
+	UI_COMMAND(RenameInterfaceParameter, "Rename Parameter", "Renames the selected interface parameter",
+			EUserInterfaceActionType::Button, FInputChord(EKeys::F2));
+	UI_COMMAND(DeleteInterfaceParameter, "Delete Parameter", "Deletes the selected interface parameter(s)",
+			EUserInterfaceActionType::Button, FInputChord(EKeys::Delete));
 }
 
 }  // namespace UE::Cameras
diff --git a/Engine/Plugins/Cameras/GameplayCameras/Source/GameplayCamerasEditor/Private/Commands/CameraRigAssetEditorCommands.cpp b/Engine/Plugins/Cameras/GameplayCameras/Source/GameplayCamerasEditor/Private/Commands/CameraRigAssetEditorCommands.cpp
--- a/Engine/Plugins/Cameras/GameplayCameras/Source/GameplayCamerasEditor/Private/Commands/CameraRigAssetEditorCommands.cpp
+++ b/Engine/Plugins/Cameras/GameplayCameras/Source/GameplayCamerasEditor/Private/Commands/CameraRigAssetEditorCommands.cpp
@@ -40,6 +40,11 @@ void FCameraRigAssetEditorCommands::RegisterCommands()
 			EUserInterfaceActionType::Button, FInputChord(EKeys::F, EModifierKey::Control));
 	UI_COMMAND(FocusHome, "Home", "Focuses the graph canvas back on the root node",
 			EUserInterfaceActionType::Button, FInputChord());
+
+	UI_COMMAND(RenameInterfaceParameter, "Rename Interface Parameter", "Renames the selected interface parameter",
+			EUserInterfaceActionType::Button, FInputChord(EKeys::F2));
+	UI_COMMAND(DeleteInterfaceParameter, "Delete Interface Parameter", "Deletes the selected interface parameter(s)",
+			EUserInterfaceActionType::Button, FInputChord(EKeys::Delete));
 }
 
 }  // namespace UE::Cameras
diff --git a/Engine/Plugins/Cameras/GameplayCameras/Source/GameplayCamerasEditor/Private/Commands/CameraShakeAssetEditorCommands.cpp b/Engine/Plugins/Cameras/GameplayCameras/Source/GameplayCamerasEditor/Private/Commands/CameraShakeAssetEditorCommands.cpp
--- a/Engine/Plugins/Cameras/GameplayCameras/Source/GameplayCamerasEditor/Private/Commands/CameraShakeAssetEditorCommands.cpp
+++ b/Engine/Plugins/Cameras/GameplayCameras/Source/GameplayCamerasEditor/Private/Commands/CameraShakeAssetEditorCommands.cpp
@@ -35,6 +35,11 @@ void FCameraShakeAssetEditorCommands::RegisterCommands()
 			EUserInterfaceActionType::Button, FInputChord(EKeys::F, EModifierKey::Control));
 	UI_COMMAND(FocusHome, "Home", "Focuses the graph canvas back on the root node",
 			EUserInterfaceActionType::Button, FInputChord());
+
+	UI_COMMAND(RenameInterfaceParameter, "Rename Interface Parameter", "Renames the selected interface parameter",
+			EUserInterfaceActionType::Button, FInputChord(EKeys::F2));
+	UI_COMMAND(DeleteInterfaceParameter, "Delete Interface Parameter", "Deletes the selected interface parameter(s)",
+			EUserInterfaceActionType::Button, FInputChord(EKeys::Delete));
 }
 
 }  // namespace UE::Cameras
diff --git a/Engine/Plugins/Cameras/GameplayCameras/Source/GameplayCamerasEditor/Private/Editors/SCameraAssetInterfaceParametersPanel.cpp b/Engine/Plugins/Cameras/GameplayCameras/Source/GameplayCamerasEditor/Private/Editors/SCameraAssetInterfaceParametersPanel.cpp
--- a/Engine/Plugins/Cameras/GameplayCameras/Source/GameplayCamerasEditor/Private/Editors/SCameraAssetInterfaceParametersPanel.cpp
+++ b/Engine/Plugins/Cameras/GameplayCameras/Source/GameplayCamerasEditor/Private/Editors/SCameraAssetInterfaceParametersPanel.cpp
@@ -490,6 +490,21 @@ void SCameraAssetInterfaceParametersPanel::Construct(const FArguments& Args, FCa
 						.ColorAndOpacity(FSlateColor::UseForeground())
 					]
 				]
+				+SHorizontalBox::Slot()
+				.AutoWidth()
+				.VAlign(VAlign_Center)
+				[
+					SNew(SButton)
+					.ButtonStyle(FAppStyle::Get(), "SimpleButton")
+					.ContentPadding(FMargin(1, 0))
+					.ToolTipText(LOCTEXT("AddParameterToolTip", "Removes the selected parameter(s)"))
+					.OnClicked(this, &SCameraAssetInterfaceParametersPanel::OnDeleteSelectedInterfaceParameter)
+					[
+						SNew(SImage)
+						.Image(FAppStyle::Get().GetBrush("Icons.MinusCircle"))
+						.ColorAndOpacity(FSlateColor::UseForeground())
+					]
+				]
 			]
 		]
 		+SVerticalBox::Slot()
@@ -529,6 +544,26 @@ void SCameraAssetInterfaceParametersPanel::RequestListRefresh()
 	bListRefreshRequested = true;
 }
 
+void SCameraAssetInterfaceParametersPanel::RenameSelectedParameter()
+{
+	if (ParametersListView->HasKeyboardFocus())
+	{
+		TArray<TObjectPtr<UCameraAssetInterfaceParameter>> SelectedItems = ParametersListView->GetSelectedItems();
+		if (SelectedItems.Num() > 0)
+		{
+			OnRenameInterfaceParameter(SelectedItems[0]);
+		}
+	}
+}
+
+void SCameraAssetInterfaceParametersPanel::DeleteSelectedParameter()
+{
+	if (ParametersListView->HasKeyboardFocus())
+	{
+		OnDeleteSelectedInterfaceParameter();
+	}
+}
+
 void SCameraAssetInterfaceParametersPanel::Tick(const FGeometry& AllottedGeometry, const double InCurrentTime, const float InDeltaTime)
 {
 	SCompoundWidget::Tick(AllottedGeometry, InCurrentTime, InDeltaTime);
@@ -588,7 +623,7 @@ FReply SCameraAssetInterfaceParametersPanel::OnAddInterfaceParameter()
 	CameraAsset->Modify();
 
 	UCameraAssetInterfaceParameter* NewParameter = NewObject<UCameraAssetInterfaceParameter>(CameraAsset, NAME_None, RF_Transactional);
-	NewParameter->InterfaceParameterName = NewParameter->GetName();
+	NewParameter->InterfaceParameterName = GetNewParameterName();
 
 	CameraAsset->Interface.Parameters.Add(NewParameter);
 	CameraAsset->EventHandlers.Notify(&ICameraAssetEventHandler::OnCameraAssetInterfaceChanged);
@@ -598,6 +633,31 @@ FReply SCameraAssetInterfaceParametersPanel::OnAddInterfaceParameter()
 	return FReply::Handled();
 }
 
+FReply SCameraAssetInterfaceParametersPanel::OnDeleteSelectedInterfaceParameter()
+{
+	TArray<TObjectPtr<UCameraAssetInterfaceParameter>> SelectedItems = ParametersListView->GetSelectedItems();
+	if (SelectedItems.IsEmpty())
+	{
+		return FReply::Handled();
+	}
+
+	const FScopedTransaction Transaction(LOCTEXT("RemoveInterfaceParameter", "Remove Interface Parameter"));
+
+	CameraAsset->Modify();
+
+	for (TObjectPtr<UCameraAssetInterfaceParameter> Item : SelectedItems)
+	{
+		const int32 NumRemoved = CameraAsset->Interface.Parameters.Remove(Item);
+		ensure(NumRemoved == 1);
+	}
+
+	CameraAsset->EventHandlers.Notify(&ICameraAssetEventHandler::OnCameraAssetInterfaceChanged);
+
+	ParametersListView->RequestListRefresh();
+
+	return FReply::Handled();
+}
+
 void SCameraAssetInterfaceParametersPanel::OnRenameInterfaceParameter(TObjectPtr<UCameraAssetInterfaceParameter> Item)
 {
 	TSharedPtr<ITableRow> RowWidget = ParametersListView->WidgetFromItem(Item);
@@ -625,6 +685,32 @@ void SCameraAssetInterfaceParametersPanel::OnDeleteInterfaceParameter(TObjectPtr
 	ParametersListView->RequestListRefresh();
 }
 
+FString SCameraAssetInterfaceParametersPanel::GetNewParameterName()
+{
+	static const FString BaseNewParameterName(TEXT("NewParameter"));
+
+	int32 NameIndex = 0;
+	bool bFoundUniqueName = false;
+	FString NewParameterName(BaseNewParameterName);
+
+	while (!bFoundUniqueName)
+	{
+		const bool bFoundConflict = CameraAsset->Interface.Parameters.ContainsByPredicate(
+				[&NewParameterName](const UCameraAssetInterfaceParameter* Item)
+				{
+					return Item->InterfaceParameterName == NewParameterName;
+				});
+		bFoundUniqueName = !bFoundConflict;
+		if (bFoundConflict)
+		{
+			++NameIndex;
+			NewParameterName = FString::Printf(TEXT("%s_%d"), *BaseNewParameterName, NameIndex);
+		}
+	}
+
+	return NewParameterName;
+}
+
 void SCameraAssetInterfaceParametersPanel::OnCameraAssetInterfaceChanged()
 {
 	bListRefreshRequested = true;
diff --git a/Engine/Plugins/Cameras/GameplayCameras/Source/GameplayCamerasEditor/Private/Editors/SCameraObjectInterfaceParametersPanel.cpp b/Engine/Plugins/Cameras/GameplayCameras/Source/GameplayCamerasEditor/Private/Editors/SCameraObjectInterfaceParametersPanel.cpp
--- a/Engine/Plugins/Cameras/GameplayCameras/Source/GameplayCamerasEditor/Private/Editors/SCameraObjectInterfaceParametersPanel.cpp
+++ b/Engine/Plugins/Cameras/GameplayCameras/Source/GameplayCamerasEditor/Private/Editors/SCameraObjectInterfaceParametersPanel.cpp
@@ -17,7 +17,6 @@
 #include "Widgets/SCompoundWidget.h"
 #include "Widgets/Text/SInlineEditableTextBlock.h"
 #include "Widgets/Text/SRichTextBlock.h"
-#include "Widgets/Text/STextBlock.h"
 #include "Widgets/Views/SListView.h"
 
 #include UE_INLINE_GENERATED_CPP_BY_NAME(SCameraObjectInterfaceParametersPanel)
@@ -501,6 +500,21 @@ void SCameraObjectInterfaceParametersPanel::Construct(const FArguments& Args, FC
 							.ColorAndOpacity(FSlateColor::UseForeground())
 						]
 					]
+					+SHorizontalBox::Slot()
+					.AutoWidth()
+					.VAlign(VAlign_Center)
+					[
+						SNew(SButton)
+						.ButtonStyle(FAppStyle::Get(), "SimpleButton")
+						.ContentPadding(FMargin(1, 0))
+						.ToolTipText(LOCTEXT("AddBlendableToolTip", "Removes the selected blendable parameter(s)"))
+						.OnClicked(this, &SCameraObjectInterfaceParametersPanel::OnDeleteSelectedBlendableParameter)
+						[
+							SNew(SImage)
+							.Image(FAppStyle::Get().GetBrush("Icons.MinusCircle"))
+							.ColorAndOpacity(FSlateColor::UseForeground())
+						]
+					]
 				]
 			]
 			+SVerticalBox::Slot()
@@ -564,6 +578,21 @@ void SCameraObjectInterfaceParametersPanel::Construct(const FArguments& Args, FC
 							.ColorAndOpacity(FSlateColor::UseForeground())
 						]
 					]
+					+SHorizontalBox::Slot()
+					.AutoWidth()
+					.VAlign(VAlign_Center)
+					[
+						SNew(SButton)
+						.ButtonStyle(FAppStyle::Get(), "SimpleButton")
+						.ContentPadding(FMargin(1, 0))
+						.ToolTipText(LOCTEXT("RemoveDataParameterToolTip", "Removes the selected data parameter(s)"))
+						.OnClicked(this, &SCameraObjectInterfaceParametersPanel::OnDeleteSelectedDataParameter)
+						[
+							SNew(SImage)
+							.Image(FAppStyle::Get().GetBrush("Icons.MinusCircle"))
+							.ColorAndOpacity(FSlateColor::UseForeground())
+						]
+					]
 				]
 			]
 			+SVerticalBox::Slot()
@@ -596,6 +625,30 @@ void SCameraObjectInterfaceParametersPanel::RequestListRefresh()
 	bListRefreshRequested = true;
 }
 
+void SCameraObjectInterfaceParametersPanel::RenameSelectedParameter()
+{
+	if (BlendableParametersListView->HasKeyboardFocus())
+	{
+		OnRenameSelectedBlendableParameter();
+	}
+	else if (DataParametersListView->HasKeyboardFocus())
+	{
+		OnRenameSelectedDataParameter();
+	}
+}
+
+void SCameraObjectInterfaceParametersPanel::DeleteSelectedParameter()
+{
+	if (BlendableParametersListView->HasKeyboardFocus())
+	{
+		OnDeleteSelectedBlendableParameter();
+	}
+	else if (DataParametersListView->HasKeyboardFocus())
+	{
+		OnDeleteSelectedDataParameter();
+	}
+}
+
 void SCameraObjectInterfaceParametersPanel::Tick(const FGeometry& AllottedGeometry, const double InCurrentTime, const float InDeltaTime)
 {
 	SCompoundWidget::Tick(AllottedGeometry, InCurrentTime, InDeltaTime);
@@ -730,12 +783,54 @@ void SCameraObjectInterfaceParametersPanel::OnDeleteInterfaceParameter(
 	ListView->RequestListRefresh();
 }
 
+template<typename ItemType>
+void SCameraObjectInterfaceParametersPanel::OnRenameSelectedInterfaceParameter(TSharedPtr<SListView<TObjectPtr<ItemType>>> ListView)
+{
+	TArray<TObjectPtr<ItemType>> SelectedItems = ListView->GetSelectedItems();
+	if (SelectedItems.Num() > 0)
+	{
+		OnRenameInterfaceParameter(ListView, SelectedItems[0]);
+	}
+}
+
+template<typename ItemType>
+void SCameraObjectInterfaceParametersPanel::OnDeleteSelectedInterfaceParameter(TSharedPtr<SListView<TObjectPtr<ItemType>>> ListView)
+{
+	TArray<TObjectPtr<ItemType>> SelectedItems = ListView->GetSelectedItems();
+	if (SelectedItems.IsEmpty())
+	{
+		return;
+	}
+
+	const FScopedTransaction Transaction(LOCTEXT("RemoveInterfaceParameter", "Remove Interface Parameter"));
+
+	CameraObject->Modify();
+
+	for (TObjectPtr<ItemType> Item : SelectedItems)
+	{
+		if constexpr (std::is_same_v<ItemType, UCameraObjectInterfaceBlendableParameter>)
+		{
+			const int32 NumRemoved = CameraObject->Interface.BlendableParameters.Remove(Item);
+			ensure(NumRemoved == 1);
+		}
+		else if constexpr (std::is_same_v<ItemType, UCameraObjectInterfaceDataParameter>)
+		{
+			const int32 NumRemoved = CameraObject->Interface.DataParameters.Remove(Item);
+			ensure(NumRemoved == 1);
+		}
+	}
+
+	CameraObject->EventHandlers.Notify(&ICameraObjectEventHandler::OnCameraObjectInterfaceChanged);
+
+	ListView->RequestListRefresh();
+}
+
 FReply SCameraObjectInterfaceParametersPanel::OnAddBlendableParameter()
 {
 	const FScopedTransaction Transaction(LOCTEXT("AddBlendableParameter", "Add Blendable Parameter"));
 
 	UCameraObjectInterfaceBlendableParameter* NewBlendableParameter = NewObject<UCameraObjectInterfaceBlendableParameter>(CameraObject, NAME_None, RF_Transactional);
-	NewBlendableParameter->InterfaceParameterName = NewBlendableParameter->GetName();
+	NewBlendableParameter->InterfaceParameterName = GetNewParameterName();
 
 	CameraObject->Modify();
 	CameraObject->Interface.BlendableParameters.Add(NewBlendableParameter);
@@ -749,7 +844,7 @@ FReply SCameraObjectInterfaceParametersPanel::OnAddDataParameter()
 	const FScopedTransaction Transaction(LOCTEXT("AddDataParameter", "Add Data Parameter"));
 
 	UCameraObjectInterfaceDataParameter* NewDataParameter = NewObject<UCameraObjectInterfaceDataParameter>(CameraObject, NAME_None, RF_Transactional);
-	NewDataParameter->InterfaceParameterName = NewDataParameter->GetName();
+	NewDataParameter->InterfaceParameterName = GetNewParameterName();
 
 	CameraObject->Modify();
 	CameraObject->Interface.DataParameters.Add(NewDataParameter);
@@ -758,6 +853,58 @@ FReply SCameraObjectInterfaceParametersPanel::OnAddDataParameter()
 	return FReply::Handled();
 }
 
+void SCameraObjectInterfaceParametersPanel::OnRenameSelectedBlendableParameter()
+{
+	OnRenameSelectedInterfaceParameter(BlendableParametersListView);
+}
+
+void SCameraObjectInterfaceParametersPanel::OnRenameSelectedDataParameter()
+{
+	OnRenameSelectedInterfaceParameter(DataParametersListView);
+}
+
+FReply SCameraObjectInterfaceParametersPanel::OnDeleteSelectedBlendableParameter()
+{
+	OnDeleteSelectedInterfaceParameter(BlendableParametersListView);
+	return FReply::Handled();
+}
+
+FReply SCameraObjectInterfaceParametersPanel::OnDeleteSelectedDataParameter()
+{
+	OnDeleteSelectedInterfaceParameter(DataParametersListView);
+	return FReply::Handled();
+}
+
+FString SCameraObjectInterfaceParametersPanel::GetNewParameterName() const
+{
+	static const FString BaseNewParameterName(TEXT("NewParameter"));
+
+	int32 NameIndex = 0;
+	bool bFoundUniqueName = false;
+	FString NewParameterName(BaseNewParameterName);
+
+	TArray<UCameraObjectInterfaceParameterBase*> AllParameters;
+	AllParameters.Append(CameraObject->Interface.BlendableParameters);
+	AllParameters.Append(CameraObject->Interface.DataParameters);
+
+	while (!bFoundUniqueName)
+	{
+		const bool bFoundConflict = AllParameters.ContainsByPredicate(
+				[&NewParameterName](const UCameraObjectInterfaceParameterBase* Item)
+				{
+					return Item->InterfaceParameterName == NewParameterName;
+				});
+		bFoundUniqueName = !bFoundConflict;
+		if (bFoundConflict)
+		{
+			++NameIndex;
+			NewParameterName = FString::Printf(TEXT("%s_%d"), *BaseNewParameterName, NameIndex);
+		}
+	}
+
+	return NewParameterName;
+}
+
 void SCameraObjectInterfaceParametersPanel::OnCameraObjectInterfaceChanged()
 {
 	bListRefreshRequested = true;
diff --git a/Engine/Plugins/Cameras/GameplayCameras/Source/GameplayCamerasEditor/Private/Toolkits/CameraAssetEditorToolkit.cpp b/Engine/Plugins/Cameras/GameplayCameras/Source/GameplayCamerasEditor/Private/Toolkits/CameraAssetEditorToolkit.cpp
--- a/Engine/Plugins/Cameras/GameplayCameras/Source/GameplayCamerasEditor/Private/Toolkits/CameraAssetEditorToolkit.cpp
+++ b/Engine/Plugins/Cameras/GameplayCameras/Source/GameplayCamerasEditor/Private/Toolkits/CameraAssetEditorToolkit.cpp
@@ -11,6 +11,7 @@
 #include "Core/CameraAsset.h"
 #include "Core/CameraDirector.h"
 #include "Core/CameraRigAsset.h"
+#include "Editor.h"
 #include "Editors/ObjectTreeGraphConfig.h"
 #include "Editors/SCameraAssetInterfaceParametersPanel.h"
 #include "Editors/SFindInObjectTreeGraph.h"
@@ -81,10 +82,19 @@ FCameraAssetEditorToolkit::FCameraAssetEditorToolkit(UCameraAssetEditor* InOwnin
 				FTabManager::FTab(InterfaceParametersTabId, ETabState::OpenedTab));
 	}
 	LayoutExtenders.Add(NewLayoutExtender);
+
+	if (GEditor)
+	{
+		GEditor->RegisterForUndo(this);
+	}
 }
 
 FCameraAssetEditorToolkit::~FCameraAssetEditorToolkit()
 {
+	if (GEditor)
+	{
+		GEditor->UnregisterForUndo(this);
+	}
 }
 
 void FCameraAssetEditorToolkit::AddReferencedObjects(FReferenceCollector& Collector)
@@ -295,6 +305,20 @@ void FCameraAssetEditorToolkit::PostInitAssetEditor()
 
 	BuildLogToolkit->OnRequestJumpToObject().BindSP(this, &FCameraAssetEditorToolkit::OnJumpToObject);
 
+	ToolkitCommands->MapAction(
+			Commands.RenameInterfaceParameter,
+			FExecuteAction::CreateSPLambda(this, [this]()
+				{
+					InterfaceParametersPanel->RenameSelectedParameter();
+				}));
+
+	ToolkitCommands->MapAction(
+			Commands.DeleteInterfaceParameter,
+			FExecuteAction::CreateSPLambda(this, [this]()
+				{
+					InterfaceParametersPanel->DeleteSelectedParameter();
+				}));
+
 	IGameplayCamerasModule& GameplayCamerasModule = IGameplayCamerasModule::Get();
 	LiveEditManager = GameplayCamerasModule.GetLiveEditManager();
 
@@ -500,6 +524,16 @@ void FCameraAssetEditorToolkit::OnInterfaceParameterSelected(UCameraAssetInterfa
 	DetailsView->SetObject(InterfaceParameter);
 }
 
+void FCameraAssetEditorToolkit::PostUndo(bool bSuccess)
+{
+	InterfaceParametersPanel->RequestListRefresh();
+}
+
+void FCameraAssetEditorToolkit::PostRedo(bool bSuccess)
+{
+	InterfaceParametersPanel->RequestListRefresh();
+}
+
 FText FCameraAssetEditorToolkit::GetBaseToolkitName() const
 {
 	return LOCTEXT("AppLabel", "Camera Asset");
diff --git a/Engine/Plugins/Cameras/GameplayCameras/Source/GameplayCamerasEditor/Private/Toolkits/CameraObjectInterfaceParametersToolkit.cpp b/Engine/Plugins/Cameras/GameplayCameras/Source/GameplayCamerasEditor/Private/Toolkits/CameraObjectInterfaceParametersToolkit.cpp
--- a/Engine/Plugins/Cameras/GameplayCameras/Source/GameplayCamerasEditor/Private/Toolkits/CameraObjectInterfaceParametersToolkit.cpp
+++ b/Engine/Plugins/Cameras/GameplayCameras/Source/GameplayCamerasEditor/Private/Toolkits/CameraObjectInterfaceParametersToolkit.cpp
@@ -64,6 +64,16 @@ void FCameraObjectInterfaceParametersToolkit::PostRedo(bool bSuccess)
 	Panel->RequestListRefresh();
 }
 
+void FCameraObjectInterfaceParametersToolkit::RenameSelectedParameter()
+{
+	Panel->RenameSelectedParameter();
+}
+
+void FCameraObjectInterfaceParametersToolkit::DeleteSelectedParameter()
+{
+	Panel->DeleteSelectedParameter();
+}
+
 }  // namespace UE::Cameras
 
 #undef LOCTEXT_NAMESPACE
diff --git a/Engine/Plugins/Cameras/GameplayCameras/Source/GameplayCamerasEditor/Private/Toolkits/CameraRigAssetEditorToolkit.cpp b/Engine/Plugins/Cameras/GameplayCameras/Source/GameplayCamerasEditor/Private/Toolkits/CameraRigAssetEditorToolkit.cpp
--- a/Engine/Plugins/Cameras/GameplayCameras/Source/GameplayCamerasEditor/Private/Toolkits/CameraRigAssetEditorToolkit.cpp
+++ b/Engine/Plugins/Cameras/GameplayCameras/Source/GameplayCamerasEditor/Private/Toolkits/CameraRigAssetEditorToolkit.cpp
@@ -305,6 +305,20 @@ void FCameraRigAssetEditorToolkit::PostInitAssetEditor()
 			CameraRigEditorWidget->FindAndJumpToObjectNode(Object);
 		});
 
+	ToolkitCommands->MapAction(
+			Commands.RenameInterfaceParameter,
+			FExecuteAction::CreateSPLambda(this, [this]()
+				{
+					InterfaceParametersToolkit->RenameSelectedParameter();
+				}));
+
+	ToolkitCommands->MapAction(
+			Commands.DeleteInterfaceParameter,
+			FExecuteAction::CreateSPLambda(this, [this]()
+				{
+					InterfaceParametersToolkit->DeleteSelectedParameter();
+				}));
+
 	IGameplayCamerasModule& GameplayCamerasModule = FModuleManager::GetModuleChecked<IGameplayCamerasModule>("GameplayCameras");
 	LiveEditManager = GameplayCamerasModule.GetLiveEditManager();
 
diff --git a/Engine/Plugins/Cameras/GameplayCameras/Source/GameplayCamerasEditor/Private/Toolkits/CameraShakeAssetEditorToolkit.cpp b/Engine/Plugins/Cameras/GameplayCameras/Source/GameplayCamerasEditor/Private/Toolkits/CameraShakeAssetEditorToolkit.cpp
--- a/Engine/Plugins/Cameras/GameplayCameras/Source/GameplayCamerasEditor/Private/Toolkits/CameraShakeAssetEditorToolkit.cpp
+++ b/Engine/Plugins/Cameras/GameplayCameras/Source/GameplayCamerasEditor/Private/Toolkits/CameraShakeAssetEditorToolkit.cpp
@@ -355,6 +355,20 @@ void FCameraShakeAssetEditorToolkit::PostInitAssetEditor()
 			OnJumpToObject(Object, NAME_None);
 		});
 
+	ToolkitCommands->MapAction(
+			Commands.RenameInterfaceParameter,
+			FExecuteAction::CreateSPLambda(this, [this]()
+				{
+					InterfaceParametersToolkit->RenameSelectedParameter();
+				}));
+
+	ToolkitCommands->MapAction(
+			Commands.DeleteInterfaceParameter,
+			FExecuteAction::CreateSPLambda(this, [this]()
+				{
+					InterfaceParametersToolkit->DeleteSelectedParameter();
+				}));
+
 	RegenerateMenusAndToolbars();
 }
 
diff --git a/Engine/Plugins/Cameras/GameplayCameras/Source/GameplayCamerasEditor/Public/Commands/CameraAssetEditorCommands.h b/Engine/Plugins/Cameras/GameplayCameras/Source/GameplayCamerasEditor/Public/Commands/CameraAssetEditorCommands.h
--- a/Engine/Plugins/Cameras/GameplayCameras/Source/GameplayCamerasEditor/Public/Commands/CameraAssetEditorCommands.h
+++ b/Engine/Plugins/Cameras/GameplayCameras/Source/GameplayCamerasEditor/Public/Commands/CameraAssetEditorCommands.h
@@ -32,6 +32,9 @@ public:
 
 	TSharedPtr<FUICommandInfo> ShowMessages;
 	TSharedPtr<FUICommandInfo> FindInCamera;
+
+	TSharedPtr<FUICommandInfo> RenameInterfaceParameter;
+	TSharedPtr<FUICommandInfo> DeleteInterfaceParameter;
 };
 
 }  // namespace UE::Cameras
diff --git a/Engine/Plugins/Cameras/GameplayCameras/Source/GameplayCamerasEditor/Public/Commands/CameraRigAssetEditorCommands.h b/Engine/Plugins/Cameras/GameplayCameras/Source/GameplayCamerasEditor/Public/Commands/CameraRigAssetEditorCommands.h
--- a/Engine/Plugins/Cameras/GameplayCameras/Source/GameplayCamerasEditor/Public/Commands/CameraRigAssetEditorCommands.h
+++ b/Engine/Plugins/Cameras/GameplayCameras/Source/GameplayCamerasEditor/Public/Commands/CameraRigAssetEditorCommands.h
@@ -30,6 +30,9 @@ public:
 	TSharedPtr<FUICommandInfo> ShowMessages;
 	TSharedPtr<FUICommandInfo> FindInCameraRig;
 	TSharedPtr<FUICommandInfo> FocusHome;
+
+	TSharedPtr<FUICommandInfo> RenameInterfaceParameter;
+	TSharedPtr<FUICommandInfo> DeleteInterfaceParameter;
 };
 
 }  // namespace UE::Cameras
diff --git a/Engine/Plugins/Cameras/GameplayCameras/Source/GameplayCamerasEditor/Public/Commands/CameraShakeAssetEditorCommands.h b/Engine/Plugins/Cameras/GameplayCameras/Source/GameplayCamerasEditor/Public/Commands/CameraShakeAssetEditorCommands.h
--- a/Engine/Plugins/Cameras/GameplayCameras/Source/GameplayCamerasEditor/Public/Commands/CameraShakeAssetEditorCommands.h
+++ b/Engine/Plugins/Cameras/GameplayCameras/Source/GameplayCamerasEditor/Public/Commands/CameraShakeAssetEditorCommands.h
@@ -27,6 +27,9 @@ public:
 	TSharedPtr<FUICommandInfo> ShowMessages;
 	TSharedPtr<FUICommandInfo> FindInCameraShake;
 	TSharedPtr<FUICommandInfo> FocusHome;
+
+	TSharedPtr<FUICommandInfo> RenameInterfaceParameter;
+	TSharedPtr<FUICommandInfo> DeleteInterfaceParameter;
 };
 
 }  // namespace UE::Cameras
diff --git a/Engine/Plugins/Cameras/GameplayCameras/Source/GameplayCamerasEditor/Public/Editors/SCameraAssetInterfaceParametersPanel.h b/Engine/Plugins/Cameras/GameplayCameras/Source/GameplayCamerasEditor/Public/Editors/SCameraAssetInterfaceParametersPanel.h
--- a/Engine/Plugins/Cameras/GameplayCameras/Source/GameplayCamerasEditor/Public/Editors/SCameraAssetInterfaceParametersPanel.h
+++ b/Engine/Plugins/Cameras/GameplayCameras/Source/GameplayCamerasEditor/Public/Editors/SCameraAssetInterfaceParametersPanel.h
@@ -40,6 +40,12 @@ public:
 	/** Request that the list of parameters be refreshed. */
 	void RequestListRefresh();
 
+	/** Rename the selected parameter in the focused panel. */
+	void RenameSelectedParameter();
+
+	/** Delete the selected parameter in the focused panel. */
+	void DeleteSelectedParameter();
+
 	/** Delegate invoked when a parmeter is selected in the panel. */
 	FOnCameraAssetInterfaceParameterEvent& OnInterfaceParameterSelected() { return OnInterfaceParameterSelectedDelegate; }
 
@@ -58,9 +64,13 @@ private:
 	TSharedPtr<SWidget> OnInterfaceParameterContextMenuOpening();
 
 	FReply OnAddInterfaceParameter();
+	FReply OnDeleteSelectedInterfaceParameter();
+
 	void OnRenameInterfaceParameter(TObjectPtr<UCameraAssetInterfaceParameter> Item);
 	void OnDeleteInterfaceParameter(TObjectPtr<UCameraAssetInterfaceParameter> Item);
 
+	FString GetNewParameterName();
+
 private:
 
 	UCameraAsset* CameraAsset = nullptr;
diff --git a/Engine/Plugins/Cameras/GameplayCameras/Source/GameplayCamerasEditor/Public/Editors/SCameraObjectInterfaceParametersPanel.h b/Engine/Plugins/Cameras/GameplayCameras/Source/GameplayCamerasEditor/Public/Editors/SCameraObjectInterfaceParametersPanel.h
--- a/Engine/Plugins/Cameras/GameplayCameras/Source/GameplayCamerasEditor/Public/Editors/SCameraObjectInterfaceParametersPanel.h
+++ b/Engine/Plugins/Cameras/GameplayCameras/Source/GameplayCamerasEditor/Public/Editors/SCameraObjectInterfaceParametersPanel.h
@@ -40,8 +40,15 @@ public:
 
 	void Construct(const FArguments& Args, FCameraObjectInterfaceParametersToolkit* OwnerToolkit);
 
+	/** Refresh the lists of parameters. */
 	void RequestListRefresh();
 
+	/** Rename the selected parameter in the focused panel. */
+	void RenameSelectedParameter();
+
+	/** Delete the selected parameter in the focused panel. */
+	void DeleteSelectedParameter();
+
 protected:
 
 	// SWidget interface.
@@ -73,9 +80,22 @@ private:
 	void OnDeleteInterfaceParameter(
 			TSharedPtr<SListView<TObjectPtr<ItemType>>> ListView, TObjectPtr<ItemType> Item);
 
+	template<typename ItemType>
+	void OnRenameSelectedInterfaceParameter(TSharedPtr<SListView<TObjectPtr<ItemType>>> ListView);
+	template<typename ItemType>
+	void OnDeleteSelectedInterfaceParameter(TSharedPtr<SListView<TObjectPtr<ItemType>>> ListView);
+
 	FReply OnAddBlendableParameter();
 	FReply OnAddDataParameter();
 
+	void OnRenameSelectedBlendableParameter();
+	void OnRenameSelectedDataParameter();
+
+	FReply OnDeleteSelectedBlendableParameter();
+	FReply OnDeleteSelectedDataParameter();
+
+	FString GetNewParameterName() const;
+
 private:
 
 	UBaseCameraObject* CameraObject = nullptr;
diff --git a/Engine/Plugins/Cameras/GameplayCameras/Source/GameplayCamerasEditor/Public/Toolkits/CameraAssetEditorToolkit.h b/Engine/Plugins/Cameras/GameplayCameras/Source/GameplayCamerasEditor/Public/Toolkits/CameraAssetEditorToolkit.h
--- a/Engine/Plugins/Cameras/GameplayCameras/Source/GameplayCamerasEditor/Public/Toolkits/CameraAssetEditorToolkit.h
+++ b/Engine/Plugins/Cameras/GameplayCameras/Source/GameplayCamerasEditor/Public/Toolkits/CameraAssetEditorToolkit.h
@@ -4,6 +4,7 @@
 
 #include "Core/CameraAsset.h"
 #include "Core/CameraEventHandler.h"
+#include "EditorUndoClient.h"
 #include "Toolkits/AssetEditorModeManagerToolkit.h"
 #include "Tools/BaseAssetToolkit.h"
 #include "UObject/GCObject.h"
@@ -31,6 +32,7 @@ class SCameraAssetInterfaceParametersPanel;
  */
 class FCameraAssetEditorToolkit 
 	: public FAssetEditorModeManagerToolkit
+	, public FEditorUndoClient
 	, public FGCObject
 	, public UE::Cameras::ICameraAssetEventHandler
 {
@@ -58,6 +60,10 @@ protected:
 	virtual FString GetWorldCentricTabPrefix() const override;
 	virtual FLinearColor GetWorldCentricTabColorScale() const override;
 
+	// FEditorUndoClient interface
+	virtual void PostUndo(bool bSuccess) override;
+	virtual void PostRedo(bool bSuccess) override;
+
 	// FGCObject interface
 	virtual void AddReferencedObjects(FReferenceCollector& Collector) override;
 	virtual FString GetReferencerName() const override;
diff --git a/Engine/Plugins/Cameras/GameplayCameras/Source/GameplayCamerasEditor/Public/Toolkits/CameraObjectInterfaceParametersToolkit.h b/Engine/Plugins/Cameras/GameplayCameras/Source/GameplayCamerasEditor/Public/Toolkits/CameraObjectInterfaceParametersToolkit.h
--- a/Engine/Plugins/Cameras/GameplayCameras/Source/GameplayCamerasEditor/Public/Toolkits/CameraObjectInterfaceParametersToolkit.h
+++ b/Engine/Plugins/Cameras/GameplayCameras/Source/GameplayCamerasEditor/Public/Toolkits/CameraObjectInterfaceParametersToolkit.h
@@ -7,6 +7,7 @@
 #include "Templates/SharedPointer.h"
 #include "UObject/ObjectPtr.h"
 
+class FUICommandList;
 class SBox;
 class SWidget;
 class UBaseCameraObject;
@@ -39,6 +40,14 @@ public:
 	/** Gets the panel widget. */
 	TSharedPtr<SWidget> GetInterfaceParametersPanel() const;
 
+public:
+
+	/** Rename the selected parameter in the focused panel. */
+	void RenameSelectedParameter();
+
+	/** Delete the selected parameter in the focused panel. */
+	void DeleteSelectedParameter();
+
 public:
 
 	/** Delegate invoked when a parmeter is selected in the panel. */
