Cameras: update unit tests following changes to camera assets

#Synced-CL 49781371

#Synced-CL 49781371

diff --git a/Engine/Plugins/Cameras/GameplayCameras/Source/GameplayCameras/Private/Tests/CameraAssetBuilderTests.cpp b/Engine/Plugins/Cameras/GameplayCameras/Source/GameplayCameras/Private/Tests/CameraAssetBuilderTests.cpp
--- a/Engine/Plugins/Cameras/GameplayCameras/Source/GameplayCameras/Private/Tests/CameraAssetBuilderTests.cpp
+++ b/Engine/Plugins/Cameras/GameplayCameras/Source/GameplayCameras/Private/Tests/CameraAssetBuilderTests.cpp
@@ -35,13 +35,6 @@ bool FCameraAssetBuilderDefaultParametersTest::RunTest(const FString& Parameters
 {
 	using namespace UE::Cameras::Test;
 
-	UFixedTestCameraDirector* CameraDirector = nullptr;
-	UCameraAsset* Camera = FCameraAssetTestBuilder()
-		.MakeDirector<UFixedTestCameraDirector>()
-			.Pin(CameraDirector)
-			.Done()
-		.Get();
-
 	UCameraRigAsset* CameraRigOne = FCameraRigAssetTestBuilder(TEXT("One"))
 		.MakeRootNode<UOffsetCameraNode>().Named(TEXT("Offset"))
 			.Done()
@@ -49,7 +42,6 @@ bool FCameraAssetBuilderDefaultParametersTest::RunTest(const FString& Parameters
 		.BuildCameraRig()
 		.SetDefaultParameterValue<FVector3d>(TEXT("TranslationOffset"), FVector3d(10.0, 20.0, 0.0))
 		.Get();
-	CameraDirector->AddCameraRig(CameraRigOne, TEXT("One"));
 
 	UCameraRigAsset* CameraRigTwo = FCameraRigAssetTestBuilder(TEXT("Two"))
 		.MakeRootNode<ULensParametersCameraNode>().Named(TEXT("Lens"))
@@ -58,6 +50,16 @@ bool FCameraAssetBuilderDefaultParametersTest::RunTest(const FString& Parameters
 		.BuildCameraRig()
 		.SetDefaultParameterValue<float>(TEXT("FocalLength"), 55.0f)
 		.Get();
+
+	UFixedTestCameraDirector* CameraDirector = nullptr;
+	UCameraAsset* Camera = FCameraAssetTestBuilder()
+		.MakeDirector<UFixedTestCameraDirector>()
+			.Pin(CameraDirector)
+			.Done()
+		.AddInterfaceParameter(CameraRigOne, TEXT("TranslationOffset"))
+		.AddInterfaceParameter(CameraRigTwo, TEXT("FocalLength"))
+		.Get();
+	CameraDirector->AddCameraRig(CameraRigOne, TEXT("One"));
 	CameraDirector->AddCameraRig(CameraRigTwo, TEXT("Two"));
 
 	Camera->BuildCamera();
@@ -83,13 +85,6 @@ bool FCameraAssetBuilderDefaultParametersCollisionTest::RunTest(const FString& P
 {
 	using namespace UE::Cameras::Test;
 
-	UFixedTestCameraDirector* CameraDirector = nullptr;
-	UCameraAsset* Camera = FCameraAssetTestBuilder()
-		.MakeDirector<UFixedTestCameraDirector>()
-			.Pin(CameraDirector)
-			.Done()
-		.Get();
-
 	UCameraRigAsset* CameraRigOne = FCameraRigAssetTestBuilder(TEXT("One"))
 		.MakeRootNode<UOffsetCameraNode>().Named(TEXT("Offset"))
 			.Done()
@@ -98,7 +93,6 @@ bool FCameraAssetBuilderDefaultParametersCollisionTest::RunTest(const FString& P
 		.BuildCameraRig()
 		.SetDefaultParameterValue<FVector3d>(TEXT("TranslationOffset"), FVector3d(10.0, 20.0, 0.0))
 		.Get();
-	CameraDirector->AddCameraRig(CameraRigOne, TEXT("One"));
 
 	UCameraRigAsset* CameraRigTwo = FCameraRigAssetTestBuilder(TEXT("Two"))
 		.MakeRootNode<ULensParametersCameraNode>().Named(TEXT("Lens"))
@@ -108,7 +102,6 @@ bool FCameraAssetBuilderDefaultParametersCollisionTest::RunTest(const FString& P
 		.BuildCameraRig()
 		.SetDefaultParameterValue<float>(TEXT("FocalLength"), 55.0f)
 		.Get();
-	CameraDirector->AddCameraRig(CameraRigTwo, TEXT("Two"));
 
 	UCameraRigAsset* CameraRigThree = FCameraRigAssetTestBuilder(TEXT("Three"))
 		.MakeArrayRootNode()
@@ -123,6 +116,21 @@ bool FCameraAssetBuilderDefaultParametersCollisionTest::RunTest(const FString& P
 		.SetDefaultParameterValue<FVector3d>(TEXT("TranslationOffset"), FVector3d(0.0, 10.0, 20.0))
 		.SetDefaultParameterValue<float>(TEXT("FocalLength"), 35.0f)
 		.Get();
+
+	UFixedTestCameraDirector* CameraDirector = nullptr;
+	UCameraAsset* Camera = FCameraAssetTestBuilder()
+		.MakeDirector<UFixedTestCameraDirector>()
+			.Pin(CameraDirector)
+			.Done()
+		.AddInterfaceParameter(CameraRigOne, TEXT("TranslationOffset"), TEXT("One_TranslationOffset"))
+		.AddInterfaceParameter(CameraRigOne, TEXT("RotationOffset"))
+		.AddInterfaceParameter(CameraRigTwo, TEXT("FocalLength"), TEXT("Two_FocalLength"))
+		.AddInterfaceParameter(CameraRigTwo, TEXT("Aperture"))
+		.AddInterfaceParameter(CameraRigThree, TEXT("TranslationOffset"), TEXT("Three_TranslationOffset"))
+		.AddInterfaceParameter(CameraRigThree, TEXT("FocalLength"), TEXT("Three_FocalLength"))
+		.Get();
+	CameraDirector->AddCameraRig(CameraRigOne, TEXT("One"));
+	CameraDirector->AddCameraRig(CameraRigTwo, TEXT("Two"));
 	CameraDirector->AddCameraRig(CameraRigThree, TEXT("Three"));
 
 	Camera->BuildCamera();
diff --git a/Engine/Plugins/Cameras/GameplayCameras/Source/GameplayCameras/Private/Tests/CameraRigAssetBuilderTests.cpp b/Engine/Plugins/Cameras/GameplayCameras/Source/GameplayCameras/Private/Tests/CameraRigAssetBuilderTests.cpp
--- a/Engine/Plugins/Cameras/GameplayCameras/Source/GameplayCameras/Private/Tests/CameraRigAssetBuilderTests.cpp
+++ b/Engine/Plugins/Cameras/GameplayCameras/Source/GameplayCameras/Private/Tests/CameraRigAssetBuilderTests.cpp
@@ -1,5 +1,6 @@
 // Copyright Epic Games, Inc. All Rights Reserved.
 
+#include "Build/CameraBuildContext.h"
 #include "Core/CameraRigAsset.h"
 #include "Misc/AutomationTest.h"
 #include "Nodes/Blends/SmoothBlendCameraNode.h"
@@ -10,6 +11,19 @@
 
 #define LOCTEXT_NAMESPACE "CameraRigAssetBuilderTests"
 
+namespace UE::Cameras::Test
+{
+
+void BuildCameraRigAsIfCooking(UCameraRigAsset* CameraRig)
+{
+	FCameraBuildLog BuildLog;
+	BuildLog.SetForwardMessagesToLogging(true);
+	FCameraBuildContext BuildContext(BuildLog, ECameraBuildReason::Cooking);
+	CameraRig->BuildCameraRig(BuildContext);
+}
+
+}  // namespace UE::Cameras::Test
+
 IMPLEMENT_SIMPLE_AUTOMATION_TEST(FCameraRigAssetBuilderNullTest, "System.Engine.GameplayCameras.CameraRigAssetBuilder.Null", EAutomationTestFlags::EditorContext | EAutomationTestFlags::EngineFilter)
 bool FCameraRigAssetBuilderNullTest::RunTest(const FString& Parameters)
 {
@@ -18,7 +32,7 @@ bool FCameraRigAssetBuilderNullTest::RunTest(const FString& Parameters)
 	UCameraRigAsset* CameraRig = FCameraRigAssetTestBuilder(TEXT("EmptyTest")).Get();
 	UTEST_EQUAL("Dirty status", CameraRig->BuildStatus, ECameraBuildStatus::Dirty);
 
-	CameraRig->BuildCameraRig();
+	UE::Cameras::Test::BuildCameraRigAsIfCooking(CameraRig);
 	UTEST_EQUAL("Clean status", CameraRig->BuildStatus, ECameraBuildStatus::Clean);
 
 	return true;
@@ -37,7 +51,7 @@ bool FCameraRigAssetBuilderSimpleAllocationTest::RunTest(const FString& Paramete
 		.Get();
 
 	UTEST_EQUAL("No evaluator allocation info", CameraRig->AllocationInfo.EvaluatorInfo.TotalSizeof, 0);
-	CameraRig->BuildCameraRig();
+	UE::Cameras::Test::BuildCameraRigAsIfCooking(CameraRig);
 
 	int32 TotalSizeof = UArrayCameraNode::GetEvaluatorAllocationInfo().Key;
 	const TTuple<int32, int32> OffsetEvaluatorInfo = UOffsetCameraNode::GetEvaluatorAllocationInfo();
@@ -62,7 +76,7 @@ bool FCameraRigAssetBuilderSimpleParameterTest::RunTest(const FString& Parameter
 		.AddBlendableParameter(TEXT("Test"), ECameraVariableType::Vector3d, OffsetNode, GET_MEMBER_NAME_CHECKED(UOffsetCameraNode, TranslationOffset))
 		.Get();
 
-	CameraRig->BuildCameraRig();
+	UE::Cameras::Test::BuildCameraRigAsIfCooking(CameraRig);
 
 	UCameraObjectInterfaceBlendableParameter* Parameter = CameraRig->Interface.BlendableParameters[0];
 	UTEST_EQUAL("Test parameter", Parameter->InterfaceParameterName, TEXT("Test"));
@@ -94,7 +108,7 @@ bool FCameraRigAssetBuilderDrivenOverridesTest::RunTest(const FString& Parameter
 		.AddBlendableParameter(TEXT("FocalLengthParam"), ECameraVariableType::Float, TEXT("Lens"), GET_MEMBER_NAME_CHECKED(ULensParametersCameraNode, FocalLength))
 		.Get();
 
-	InnerCameraRig->BuildCameraRig();
+	UE::Cameras::Test::BuildCameraRigAsIfCooking(InnerCameraRig);
 
 	// Make a camera rig that uses the previous one, with overrides on both the offset (now 15, 25, 35)
 	// and the focal length (now 25). Expose the offset further up as an interface parameter.
@@ -125,7 +139,7 @@ bool FCameraRigAssetBuilderDrivenOverridesTest::RunTest(const FString& Parameter
 			.AddBlendableParameter(TEXT("MiddleOffsetParam"), ECameraVariableType::Vector3d, MiddlePrefabNode, TEXT("OffsetParam"))
 		.Get();
 
-	MiddleCameraRig->BuildCameraRig();
+	UE::Cameras::Test::BuildCameraRigAsIfCooking(MiddleCameraRig);
 
 	// Make another camera rig that uses the previous one, which makes a total of 3 nesting levels of camera rigs.
 	// This level overrides the offset parameter some more (now 20, 50, 70).
@@ -149,7 +163,7 @@ bool FCameraRigAssetBuilderDrivenOverridesTest::RunTest(const FString& Parameter
 			.Done()
 		.Get();
 
-	OuterCameraRig->BuildCameraRig();
+	UE::Cameras::Test::BuildCameraRigAsIfCooking(OuterCameraRig);
 
 	UCameraObjectInterfaceBlendableParameter* OffsetParam = InnerCameraRig->Interface.BlendableParameters[0];
 	UCameraObjectInterfaceBlendableParameter* FocalLengthParam = InnerCameraRig->Interface.BlendableParameters[1];
diff --git a/Engine/Plugins/Cameras/GameplayCameras/Source/GameplayCameras/Private/Tests/GameplayCamerasTestBuilder.h b/Engine/Plugins/Cameras/GameplayCameras/Source/GameplayCameras/Private/Tests/GameplayCamerasTestBuilder.h
--- a/Engine/Plugins/Cameras/GameplayCameras/Source/GameplayCameras/Private/Tests/GameplayCamerasTestBuilder.h
+++ b/Engine/Plugins/Cameras/GameplayCameras/Source/GameplayCameras/Private/Tests/GameplayCamerasTestBuilder.h
@@ -802,6 +802,44 @@ public:
 		return DirectorBuilder;
 	}
 
+	/**
+	 * Creates a new interface parameter that exposes the given camera rig parameter, optionally
+	 * under a different name.
+	 */
+	ThisType& AddInterfaceParameter(UCameraRigAsset* SourceCameraRig, const FName SourceParameterName, const FString& InterfaceParameterName)
+	{
+		if (ensure(SourceCameraRig))
+		{
+			UCameraAssetInterfaceParameter* InterfaceParameter = NewObject<UCameraAssetInterfaceParameter>(CameraAsset);
+			InterfaceParameter->SourceCameraRig = SourceCameraRig;
+			InterfaceParameter->SourceParameterName = SourceParameterName;
+			InterfaceParameter->InterfaceParameterName = (
+					InterfaceParameterName.IsEmpty() ?
+					SourceParameterName.ToString() : InterfaceParameterName);
+
+			CameraAsset->Interface.Parameters.Add(InterfaceParameter);
+		}
+		return *static_cast<ThisType*>(this);
+	}
+
+	ThisType& AddInterfaceParameter(UCameraRigAsset* SourceCameraRig, const FName SourceParameterName)
+	{
+		return AddInterfaceParameter(SourceCameraRig, SourceParameterName, SourceParameterName.ToString());
+	}
+
+	ThisType& AddInterfaceParameter(const FString& SourceCameraRigName, const FName& SourceParameterName, const FString& InterfaceParameterName)
+	{
+		return AddInterfaceParameter(
+				NamedObjectRegistry->Get<UCameraRigAsset>(SourceCameraRigName),
+				SourceParameterName,
+				InterfaceParameterName);
+	}
+
+	ThisType& AddInterfaceParameter(const FString& SourceCameraRigName, const FName SourceParameterName)
+	{
+		return AddInterfaceParameter(SourceCameraRigName, SourceParameterName, SourceParameterName.ToString());
+	}
+
 	/** Gets the named object registry. */
 	virtual TSharedPtr<FNamedObjectRegistry> GetNamedObjectRegistry() override
 	{
