Cameras: fix assert on running a Camera Rig component

This is a fallout of a previous change -- the generated Camera Asset in a Camera Rig component didn't have the correct allocation info anymore for initializing the evaluation context's variable and data tables.

#Synced-CL 49781371

diff --git a/Engine/Plugins/Cameras/GameplayCameras/Source/GameplayCameras/Private/GameFramework/GameplayCameraRigComponent.cpp b/Engine/Plugins/Cameras/GameplayCameras/Source/GameplayCameras/Private/GameFramework/GameplayCameraRigComponent.cpp
--- a/Engine/Plugins/Cameras/GameplayCameras/Source/GameplayCameras/Private/GameFramework/GameplayCameraRigComponent.cpp
+++ b/Engine/Plugins/Cameras/GameplayCameras/Source/GameplayCameras/Private/GameFramework/GameplayCameraRigComponent.cpp
@@ -101,7 +101,12 @@ void UGameplayCameraRigComponent::BuildGeneratedCamera()
 
 		FCameraBuildLog BuildLog;
 		BuildLog.SetForwardMessagesToLogging(true);
-		FCameraBuildContext BuildContext(BuildLog);
+
+		// Build the asset as if we were cooking so we get the referenced camera rig's allocation info
+		// accumulated into ours. It's more efficient and more importantly it lets us initialize the
+		// evaluation context's variable and context data tables with all the entries we need for later
+		// calling CameraRigReference.EnsureAllocationInfo().
+		FCameraBuildContext BuildContext(BuildLog, ECameraBuildReason::Cooking);
 		FCameraAssetBuilder Builder(BuildContext);
 		Builder.BuildCamera(GeneratedCameraAsset, bBuildReferencedAssets);
 	}
