Cameras: fix incorrect pragmas

#Synced-CL 49781371

diff --git a/Engine/Plugins/Cameras/GameplayCameras/Source/GameplayCamerasUncookedOnly/Private/Helpers/CameraContextDataPinTypeHelper.cpp b/Engine/Plugins/Cameras/GameplayCameras/Source/GameplayCamerasUncookedOnly/Private/Helpers/CameraContextDataPinTypeHelper.cpp
--- a/Engine/Plugins/Cameras/GameplayCameras/Source/GameplayCamerasUncookedOnly/Private/Helpers/CameraContextDataPinTypeHelper.cpp
+++ b/Engine/Plugins/Cameras/GameplayCameras/Source/GameplayCamerasUncookedOnly/Private/Helpers/CameraContextDataPinTypeHelper.cpp
@@ -1,7 +1,5 @@
 // Copyright Epic Games, Inc. All Rights Reserved.
 
-#pragma once
-
 #include "Helpers/CameraContextDataPinTypeHelper.h"
 
 #include "EdGraphSchema_K2.h"
diff --git a/Engine/Plugins/Cameras/GameplayCameras/Source/GameplayCamerasUncookedOnly/Private/Helpers/CameraVariablePinTypeHelper.cpp b/Engine/Plugins/Cameras/GameplayCameras/Source/GameplayCamerasUncookedOnly/Private/Helpers/CameraVariablePinTypeHelper.cpp
--- a/Engine/Plugins/Cameras/GameplayCameras/Source/GameplayCamerasUncookedOnly/Private/Helpers/CameraVariablePinTypeHelper.cpp
+++ b/Engine/Plugins/Cameras/GameplayCameras/Source/GameplayCamerasUncookedOnly/Private/Helpers/CameraVariablePinTypeHelper.cpp
@@ -1,7 +1,5 @@
 // Copyright Epic Games, Inc. All Rights Reserved.
 
-#pragma once
-
 #include "Helpers/CameraVariablePinTypeHelper.h"
 
 #include "EdGraphSchema_K2.h"
