Cameras: remove calls to Modify() from PostLoad()

#Synced-CL 49781371

diff --git a/Engine/Plugins/Cameras/GameplayCameras/Source/GameplayCameras/Private/Core/BaseCameraObject.cpp b/Engine/Plugins/Cameras/GameplayCameras/Source/GameplayCameras/Private/Core/BaseCameraObject.cpp
--- a/Engine/Plugins/Cameras/GameplayCameras/Source/GameplayCameras/Private/Core/BaseCameraObject.cpp
+++ b/Engine/Plugins/Cameras/GameplayCameras/Source/GameplayCameras/Private/Core/BaseCameraObject.cpp
@@ -145,7 +145,6 @@ void UBaseCameraObject::UpgradeInterfaceConnections(IObjectTreeGraphRootObject*
 
 	// Move the connection information to the new connection list.
 	// Create getter nodes for parameters that were added to the camera node graph.
-	bool bDidModify = false;
 	for (UCameraObjectInterfaceParameterBase* Parameter : AllParameters)
 	{
 		UCameraObjectInterfaceParameterGetter* GetterNode = nullptr;
@@ -157,7 +156,6 @@ void UBaseCameraObject::UpgradeInterfaceConnections(IObjectTreeGraphRootObject*
 			GetterNode->ParameterGuid = Parameter->GetGuid();
 			GetterNode->GraphNodePos = Parameter->GraphNodePos_DEPRECATED;
 			RootObject->AddConnectableObject(DefaultGraphName, GetterNode);
-			bDidModify = true;
 
 			Parameter->bHasGraphNode_DEPRECATED = false;
 			Parameter->GraphNodePos_DEPRECATED = FIntVector2::ZeroValue;
@@ -173,17 +171,11 @@ void UBaseCameraObject::UpgradeInterfaceConnections(IObjectTreeGraphRootObject*
 			Connection.TargetPropertyName = Parameter->TargetPropertyName_DEPRECATED;
 
 			Connections.Connections.Add(Connection);
-			bDidModify = true;
 
 			Parameter->Target_DEPRECATED = nullptr;
 			Parameter->TargetPropertyName_DEPRECATED = NAME_None;
 		}
 	}
-
-	if (bDidModify)
-	{
-		Modify();
-	}
 }
 
 #endif  // WITH_EDITORONLY_DATA
diff --git a/Engine/Plugins/Cameras/GameplayCameras/Source/GameplayCameras/Private/Core/CameraVariableCollection.cpp b/Engine/Plugins/Cameras/GameplayCameras/Source/GameplayCameras/Private/Core/CameraVariableCollection.cpp
--- a/Engine/Plugins/Cameras/GameplayCameras/Source/GameplayCameras/Private/Core/CameraVariableCollection.cpp
+++ b/Engine/Plugins/Cameras/GameplayCameras/Source/GameplayCameras/Private/Core/CameraVariableCollection.cpp
@@ -58,8 +58,6 @@ void UCameraVariableCollection::CleanUpStrayObjects()
 			continue;
 		}
 
-		Modify();
-
 		Variable->ClearFlags(RF_Public | RF_Standalone);
 		StrayObjects.Add(Variable);
 	}
