Cameras: fix missing code path for connecting to struct data properties

#Synced-CL 49781371

diff --git a/Engine/Plugins/Cameras/GameplayCameras/Source/GameplayCameras/Private/Core/ICustomCameraNodeParameterProvider.cpp b/Engine/Plugins/Cameras/GameplayCameras/Source/GameplayCameras/Private/Core/ICustomCameraNodeParameterProvider.cpp
--- a/Engine/Plugins/Cameras/GameplayCameras/Source/GameplayCameras/Private/Core/ICustomCameraNodeParameterProvider.cpp
+++ b/Engine/Plugins/Cameras/GameplayCameras/Source/GameplayCameras/Private/Core/ICustomCameraNodeParameterProvider.cpp
@@ -262,6 +262,11 @@ void FCameraNodeParameterInfos::BuildFrom(UCameraNode* InCameraNode)
 				DataPropertyType = ECameraContextDataType::Enum;
 				DataPropertyTypeObject = EnumProperty->GetEnum();
 			}
+			else if (FStructProperty* StructProperty = CastField<FStructProperty>(ActualProperty))
+			{
+				DataPropertyType = ECameraContextDataType::Struct;
+				DataPropertyTypeObject = StructProperty->Struct;
+			}
 			else if (FClassProperty* ClassProperty = CastField<FClassProperty>(ActualProperty))
 			{
 				DataPropertyType = ECameraContextDataType::Class;
