Fix localization warnings

#Synced-CL 49841208

diff --git a/Engine/Plugins/Cameras/GameplayCameras/Source/GameplayCamerasEditor/Private/Editors/SCameraAssetInterfaceParametersPanel.cpp b/Engine/Plugins/Cameras/GameplayCameras/Source/GameplayCamerasEditor/Private/Editors/SCameraAssetInterfaceParametersPanel.cpp
--- a/Engine/Plugins/Cameras/GameplayCameras/Source/GameplayCamerasEditor/Private/Editors/SCameraAssetInterfaceParametersPanel.cpp
+++ b/Engine/Plugins/Cameras/GameplayCameras/Source/GameplayCamerasEditor/Private/Editors/SCameraAssetInterfaceParametersPanel.cpp
@@ -497,7 +497,7 @@ void SCameraAssetInterfaceParametersPanel::Construct(const FArguments& Args, FCa
 					SNew(SButton)
 					.ButtonStyle(FAppStyle::Get(), "SimpleButton")
 					.ContentPadding(FMargin(1, 0))
-					.ToolTipText(LOCTEXT("AddParameterToolTip", "Removes the selected parameter(s)"))
+					.ToolTipText(LOCTEXT("RemoveParameterToolTip", "Removes the selected parameter(s)"))
 					.OnClicked(this, &SCameraAssetInterfaceParametersPanel::OnDeleteSelectedInterfaceParameter)
 					[
 						SNew(SImage)
diff --git a/Engine/Plugins/Cameras/GameplayCameras/Source/GameplayCamerasEditor/Private/Editors/SCameraObjectInterfaceParametersPanel.cpp b/Engine/Plugins/Cameras/GameplayCameras/Source/GameplayCamerasEditor/Private/Editors/SCameraObjectInterfaceParametersPanel.cpp
--- a/Engine/Plugins/Cameras/GameplayCameras/Source/GameplayCamerasEditor/Private/Editors/SCameraObjectInterfaceParametersPanel.cpp
+++ b/Engine/Plugins/Cameras/GameplayCameras/Source/GameplayCamerasEditor/Private/Editors/SCameraObjectInterfaceParametersPanel.cpp
@@ -507,7 +507,7 @@ void SCameraObjectInterfaceParametersPanel::Construct(const FArguments& Args, FC
 						SNew(SButton)
 						.ButtonStyle(FAppStyle::Get(), "SimpleButton")
 						.ContentPadding(FMargin(1, 0))
-						.ToolTipText(LOCTEXT("AddBlendableToolTip", "Removes the selected blendable parameter(s)"))
+						.ToolTipText(LOCTEXT("RemoveBlendableToolTip", "Removes the selected blendable parameter(s)"))
 						.OnClicked(this, &SCameraObjectInterfaceParametersPanel::OnDeleteSelectedBlendableParameter)
 						[
 							SNew(SImage)
