Cameras: fix assert sometimes firing after rebuilding a camera rig

When a Camera Rig Component is selected in the level editor, it shows up in the Details View and the customization for the Camera Rig reference is active. When the referenced Camera Rig is rebuilt, that customization automatically updates the reference's parameters property bag. So when the component itself ends up receiving the build notification, the reference's parameters are up to date. This prevented the component from reinitializing the evaluation context, and could lead to asserts if new parameters were added, since those entries wouldn't exist in the variable or context data tables.

#Synced-CL 49781371

diff --git a/Engine/Plugins/Cameras/GameplayCameras/Source/GameplayCameras/Private/GameFramework/GameplayCameraRigComponent.cpp b/Engine/Plugins/Cameras/GameplayCameras/Source/GameplayCameras/Private/GameFramework/GameplayCameraRigComponent.cpp
--- a/Engine/Plugins/Cameras/GameplayCameras/Source/GameplayCameras/Private/GameFramework/GameplayCameraRigComponent.cpp
+++ b/Engine/Plugins/Cameras/GameplayCameras/Source/GameplayCameras/Private/GameFramework/GameplayCameraRigComponent.cpp
@@ -116,20 +116,17 @@ void UGameplayCameraRigComponent::OnCameraRigAssetBuiltImpl()
 {
 	using namespace UE::Cameras;
 
-	if (CameraRigReference.NeedsRebuildParameters())
+	CameraRigReference.RebuildParametersIfNeeded();
+
+	if (HasEvaluationContext())
 	{
-		Modify();
-		CameraRigReference.RebuildParameters();
-		if (HasEvaluationContext())
-		{
 #if WITH_EDITOR
-			const UCameraRigAsset* CameraRigAsset = CameraRigReference.GetCameraRig();
-			const FCameraObjectAllocationInfo& AllocationInfo = CameraRigAsset->AllocationInfo;
-			ReinitializeEvaluationContext(AllocationInfo.VariableTableInfo, AllocationInfo.ContextDataTableInfo);
+		const UCameraRigAsset* CameraRigAsset = CameraRigReference.GetCameraRig();
+		const FCameraObjectAllocationInfo& AllocationInfo = CameraRigAsset->AllocationInfo;
+		ReinitializeEvaluationContext(AllocationInfo.VariableTableInfo, AllocationInfo.ContextDataTableInfo);
 #endif  // WITH_EDITOR
 
-			UpdateEvaluationContext(true);
-		}
+		UpdateEvaluationContext(true);
 	}
 }
 
