Cameras: export some symbols for 3rd party code

#Synced-CL 49781371

diff --git a/Engine/Plugins/Cameras/GameplayCameras/Source/GameplayCameras/Private/Nodes/Input/InputAxisBindingHelpers.h b/Engine/Plugins/Cameras/GameplayCameras/Source/GameplayCameras/Private/Nodes/Input/InputAxisBindingHelpers.h
--- a/Engine/Plugins/Cameras/GameplayCameras/Source/GameplayCameras/Private/Nodes/Input/InputAxisBindingHelpers.h
+++ b/Engine/Plugins/Cameras/GameplayCameras/Source/GameplayCameras/Private/Nodes/Input/InputAxisBindingHelpers.h
@@ -5,6 +5,8 @@
 #include "CoreTypes.h"
 #include "Containers/Array.h"
 
+#define UE_API GAMEPLAYCAMERAS_API
+
 class UCameraNode;
 class UEnhancedInputComponent;
 class UInputAction;
@@ -19,17 +21,19 @@ class FInputAxisBindingHelpers
 {
 public:
 
-	static UEnhancedInputComponent* FindInputComponent(const UE::Cameras::FCameraNodeEvaluatorInitializeParams& Params);
+	UE_API static UEnhancedInputComponent* FindInputComponent(const UE::Cameras::FCameraNodeEvaluatorInitializeParams& Params);
 
-	static void BindActionValues(
+	UE_API static void BindActionValues(
 			const UE::Cameras::FCameraNodeEvaluatorInitializeParams& Params,
 			const UCameraNode* CameraNode,
 			UEnhancedInputComponent* InputComponent,
 			const TArray<TObjectPtr<UInputAction>>& AxisActions,
 			TArray<FEnhancedInputActionValueBinding*>& OutAxisValueBindings);
 
-	static FVector2d GetHighestValue(const TArray<FEnhancedInputActionValueBinding*>& AxisValueBindings);
+	UE_API static FVector2d GetHighestValue(const TArray<FEnhancedInputActionValueBinding*>& AxisValueBindings);
 };
 
 }  // namespace UE::Cameras
 
+#undef UE_API
+
diff --git a/Engine/Plugins/Cameras/GameplayCameras/Source/GameplayCameras/Public/Math/CameraFramingMath.h b/Engine/Plugins/Cameras/GameplayCameras/Source/GameplayCameras/Public/Math/CameraFramingMath.h
--- a/Engine/Plugins/Cameras/GameplayCameras/Source/GameplayCameras/Public/Math/CameraFramingMath.h
+++ b/Engine/Plugins/Cameras/GameplayCameras/Source/GameplayCameras/Public/Math/CameraFramingMath.h
@@ -5,6 +5,8 @@
 #include "CoreTypes.h"
 #include "Math/MathFwd.h"
 
+#define UE_API GAMEPLAYCAMERAS_API
+
 namespace UE::Cameras
 {
 
@@ -19,15 +21,17 @@ class FCameraFramingMath
 {
 public:
 
-	static FVector2d GetTargetAngles(const FVector2d& Target, const FCameraFieldsOfView& FieldsOfView);
+	UE_API static FVector2d GetTargetAngles(const FVector2d& Target, const FCameraFieldsOfView& FieldsOfView);
 
 	/** Gets the framing zone's half-angles for a given camera FOV. */
-	static FFramingZoneAngles GetFramingZoneAngles(const FFramingZone& FramingZone, const FCameraFieldsOfView& FieldsOfView);
+	UE_API static FFramingZoneAngles GetFramingZoneAngles(const FFramingZone& FramingZone, const FCameraFieldsOfView& FieldsOfView);
 
 private:
 
-	static double GetBoundAngle(float FactorFromCenter, double TanHalfFOV);
+	UE_API static double GetBoundAngle(float FactorFromCenter, double TanHalfFOV);
 };
 
 }  // namespace UE::Cameras
 
+#undef UE_API
+
diff --git a/Engine/Plugins/Cameras/GameplayCameras/Source/GameplayCameras/Public/Math/CameraFramingZoneMath.h b/Engine/Plugins/Cameras/GameplayCameras/Source/GameplayCameras/Public/Math/CameraFramingZoneMath.h
--- a/Engine/Plugins/Cameras/GameplayCameras/Source/GameplayCameras/Public/Math/CameraFramingZoneMath.h
+++ b/Engine/Plugins/Cameras/GameplayCameras/Source/GameplayCameras/Public/Math/CameraFramingZoneMath.h
@@ -5,6 +5,8 @@
 #include "Containers/ContainersFwd.h"
 #include "Math/UnrealMath.h"
 
+#define UE_API GAMEPLAYCAMERAS_API
+
 class FArchive;
 struct FCameraFramingZone;
 struct FCameraPose;
@@ -27,7 +29,7 @@ struct FFramingZone
 	double BottomBound = 0;
 
 	/** Builds an empty framing zone. */
-	FFramingZone();
+	UE_API FFramingZone();
 
 	/** Gets the width of the framing zone. */
 	double Width() const { return FMath::Max(0.0, RightBound - LeftBound); }
@@ -42,63 +44,63 @@ struct FFramingZone
 	 * Returns whether this framing zone is well-formed (bounds in 0..1 UI space, 
 	 * with left/top bounds lesser than or equal to right/bottom bounds).
 	 */
-	bool IsValid() const;
+	UE_API bool IsValid() const;
 
 	/** Makes sure all the bounds have valid values between 0 and 1. */
-	void ClampBounds();
+	UE_API void ClampBounds();
 
 	/**
 	 * Makes sure all the bounds have valid values between 0 and 1, and that the
 	 * enclosed rectangle contains the given target point.
 	 */
-	void ClampBounds(const FVector2d& MustContain);
+	UE_API void ClampBounds(const FVector2d& MustContain);
 
 	/**
 	 * Makes sure all the bounds have valid values between 0 and 1, and that the
 	 * enclosed rectangle contains the given target point, with a uniform margin.
 	 */
-	void ClampBounds(const FVector2d& MustContain, double Margin);
+	UE_API void ClampBounds(const FVector2d& MustContain, double Margin);
 
 	/**
 	 * Makes sure all the bounds have valid values between 0 and 1, and that the
 	 * enclosed rectangle contains the given inner rectangle.
 	 */
-	void ClampBounds(const FFramingZone& MustContain);
+	UE_API void ClampBounds(const FFramingZone& MustContain);
 
 	/**
 	 * Expands this framing zone to include the other framing zone.
 	 */
-	void Add(const FFramingZone& Other);
+	UE_API void Add(const FFramingZone& Other);
 
 	/** Checks whether the given point (in 0..1 UI space) is inside this zone. */
-	bool Contains(const FVector2d& Point) const;
+	UE_API bool Contains(const FVector2d& Point) const;
 
 	/** 
 	 * Computes intersections between a line and the zone's box, and returns the one closest
 	 * to the line's origin point.
 	 */
-	FVector2d ComputeClosestIntersection(const FVector2d& Origin, const FVector2d& LineDir, bool bLineDirIsNormalized = false) const;
+	UE_API FVector2d ComputeClosestIntersection(const FVector2d& Origin, const FVector2d& LineDir, bool bLineDirIsNormalized = false) const;
 
 	/** Gets the inner margins of this zone compared to the screen's center. */
-	FVector4d GetNormalizedBounds() const;
+	UE_API FVector4d GetNormalizedBounds() const;
 
 	/** Gets the coordinates of the top-left corner of the zone, in 0..Width/Height canvas units. */
-	FVector2d GetCanvasPosition(const FVector2d& CanvasSize) const;
+	UE_API FVector2d GetCanvasPosition(const FVector2d& CanvasSize) const;
 	/** Gets the size of the zone, in 0..Width/Height canvas units. */
-	FVector2d GetCanvasSize(const FVector2d& CanvasSize) const;
+	UE_API FVector2d GetCanvasSize(const FVector2d& CanvasSize) const;
 
-	void Serialize(FArchive& Ar);
+	UE_API void Serialize(FArchive& Ar);
 
 public:
 
 	/** Build a framing zone from a set of margins relative to the screen edge. */
-	static FFramingZone FromScreenMargins(const FCameraFramingZone& Margins);
+	UE_API static FFramingZone FromScreenMargins(const FCameraFramingZone& Margins);
 
 	/** Build a framing zone from a set of margins relative to a screen location. */
-	static FFramingZone FromRelativeMargins(const FVector2d& ScreenLocation, const FCameraFramingZone& Margins);
+	UE_API static FFramingZone FromRelativeMargins(const FVector2d& ScreenLocation, const FCameraFramingZone& Margins);
 
 	/** Build a framing zone that encompasses all the given points (in 0..1 UI space). */
-	static FFramingZone FromPoints(TConstArrayView<FVector2d> ScreenPoints);
+	UE_API static FFramingZone FromPoints(TConstArrayView<FVector2d> ScreenPoints);
 
 private:
 
@@ -121,3 +123,5 @@ struct FFramingZoneAngles
 
 }  // namespace UE::Cameras
 
+#undef UE_API
+
diff --git a/Engine/Plugins/Cameras/GameplayCameras/Source/GameplayCameras/Public/Math/CameraPoseMath.h b/Engine/Plugins/Cameras/GameplayCameras/Source/GameplayCameras/Public/Math/CameraPoseMath.h
--- a/Engine/Plugins/Cameras/GameplayCameras/Source/GameplayCameras/Public/Math/CameraPoseMath.h
+++ b/Engine/Plugins/Cameras/GameplayCameras/Source/GameplayCameras/Public/Math/CameraPoseMath.h
@@ -6,6 +6,8 @@
 #include "Math/MathFwd.h"
 #include "Math/Matrix.h"
 
+#define UE_API GAMEPLAYCAMERAS_API
+
 struct FCameraFieldsOfView;
 struct FCameraPose;
 
@@ -31,35 +33,35 @@ class FCameraPoseMath
 public:
 
 	/** Gets both horizontal and vertical effective fields of view, using the sensor aspect ratio. */
-	static FCameraFieldsOfView GetEffectiveFieldsOfView(const FCameraPose& CameraPose);
+	UE_API static FCameraFieldsOfView GetEffectiveFieldsOfView(const FCameraPose& CameraPose);
 
 	/** Gets both horizontal and vertical effective fields of view, using the given player controller's viewport aspect ratio. */
-	static FCameraFieldsOfView GetEffectiveFieldsOfView(const FCameraPose& CameraPose, TSharedPtr<const FCameraEvaluationContext> EvaluationContext);
+	UE_API static FCameraFieldsOfView GetEffectiveFieldsOfView(const FCameraPose& CameraPose, TSharedPtr<const FCameraEvaluationContext> EvaluationContext);
 
 	/** Gets both horizontal and vertical effective fields of view, using the given aspect ratio. */
-	static FCameraFieldsOfView GetEffectiveFieldsOfView(const FCameraPose& CameraPose, double AspectRatio);
+	UE_API static FCameraFieldsOfView GetEffectiveFieldsOfView(const FCameraPose& CameraPose, double AspectRatio);
 
 	/** Gets the aspect ratio of the viewport associated with the given player controller. */
-	static double GetEffectiveAspectRatio(const FCameraPose& CameraPose, TSharedPtr<const FCameraEvaluationContext> EvaluationContext);
+	UE_API static double GetEffectiveAspectRatio(const FCameraPose& CameraPose, TSharedPtr<const FCameraEvaluationContext> EvaluationContext);
 
 	/** 
 	 * Builds the projection matrix of the given camera pose. 
 	 * This matrix is suitable for projecting camera-space points onto screen-space.
 	 */
-	static FMatrix BuildProjectionMatrix(const FCameraPose& CameraPose, double AspectRatio);
+	UE_API static FMatrix BuildProjectionMatrix(const FCameraPose& CameraPose, double AspectRatio);
 	/** 
 	 * Builds the view-projection matrix of the given camera pose.
 	 * This matrix combines the camera transform and the projection matrix, making it
 	 * suitable for projecting world-space points onto screen-space.
 	 */
-	static FMatrix BuildViewProjectionMatrix(const FCameraPose& CameraPose, double AspectRatio);
+	UE_API static FMatrix BuildViewProjectionMatrix(const FCameraPose& CameraPose, double AspectRatio);
 
 	/** Projects the given world point onto screen-space. */
-	static TOptional<FVector2d> ProjectWorldToScreen(
+	UE_API static TOptional<FVector2d> ProjectWorldToScreen(
 			const FCameraPose& CameraPose, double AspectRatio, 
 			const FVector3d& WorldLocation, bool bForceLocationInsideFrustum = false);
 	/** Projects the given camera-local point onto screen-space. */
-	static TOptional<FVector2d> ProjectCameraToScreen(
+	UE_API static TOptional<FVector2d> ProjectCameraToScreen(
 			const FCameraPose& CameraPose, double AspectRatio, 
 			const FVector3d& CameraSpaceLocation, bool bForceLocationInsideFrustum = false);
 
@@ -68,27 +70,26 @@ public:
 	 * The caller is responsible for making sure the view-projection matrix and coordinate system 
 	 * for the given point are compatible.
 	 */
-	static TOptional<FVector2d> ProjectToScreen(
+	UE_API static TOptional<FVector2d> ProjectToScreen(
 			const FMatrix& ViewProjectionMatrix, const FVector3d& Location, bool bForceLocationInsideFrustum = false);
 
 	/** Unproject a screen-space point into a camera-local ray. */
-	static FRay3d UnprojectScreenToCamera(
+	UE_API static FRay3d UnprojectScreenToCamera(
 			const FCameraPose& CameraPose, double AspectRatio, const FVector2D& ScreenSpacePoint);
 	/** Unproject a screen-space point into a camera-local point given an expected distance. */
-	static FVector3d UnprojectScreenToCamera(
+	UE_API static FVector3d UnprojectScreenToCamera(
 			const FCameraPose& CameraPose, double AspectRatio, const FVector2D& ScreenSpacePoint, double PredictedDistance);
 	/** Unproject a screen-space point into a world-space ray. */
-	static FRay3d UnprojectScreenToWorld(
+	UE_API static FRay3d UnprojectScreenToWorld(
 			const FCameraPose& CameraPose, double AspectRatio, const FVector2D& ScreenSpacePoint);
 	/** Unproject a screen-space point into a world-space point given an expected distance. */
-	static FVector3d UnprojectScreenToWorld(
+	UE_API static FVector3d UnprojectScreenToWorld(
 			const FCameraPose& CameraPose, double AspectRatio, const FVector2D& ScreenSpacePoint, double PredictedDistance);
-
 	/** Unprojects the given screen-space point into a ray. */
-	static FRay3d UnprojectFromScreen(
+	UE_API static FRay3d UnprojectFromScreen(
 			const FMatrix& InverseViewProjectionMatrix, const FVector2D& ScreenSpacePoint);
 	/** Unprojects the given screen-space point into a point given an expected distance. */
-	static FVector3d UnprojectFromScreen(
+	UE_API static FVector3d UnprojectFromScreen(
 			const FMatrix& InverseViewProjectionMatrix, const FVector2D& ScreenSpacePoint, double PredictedDistance);
 
 private:
@@ -98,3 +99,5 @@ private:
 
 }  // namespace UE::Cameras
 
+#undef UE_API
+
