Cameras: make sure the "aim at" action uses the shorter rotation to accomplish its goal

#Synced-CL 49781371

diff --git a/Engine/Plugins/Cameras/GameplayCameras/Source/GameplayCameras/Private/Actions/AimAtCameraAction.cpp b/Engine/Plugins/Cameras/GameplayCameras/Source/GameplayCameras/Private/Actions/AimAtCameraAction.cpp
--- a/Engine/Plugins/Cameras/GameplayCameras/Source/GameplayCameras/Private/Actions/AimAtCameraAction.cpp
+++ b/Engine/Plugins/Cameras/GameplayCameras/Source/GameplayCameras/Private/Actions/AimAtCameraAction.cpp
@@ -184,7 +184,7 @@ bool FAimAtCameraActionEvaluator::RunPreviewEvaluation(const FCameraActionEvalua
 		return false;
 	}
 
-	OutCorrection = TotalCorrection;
+	OutCorrection = TotalCorrection.GetNormalized();
 	return true;
 }
 
