Cameras: fix regression bug about camera node parameters being incorrectly built when they use a camera variable asset

#Synced-CL 49781371

diff --git a/Engine/Plugins/Cameras/GameplayCameras/Source/GameplayCameras/Private/Core/ICustomCameraNodeParameterProvider.cpp b/Engine/Plugins/Cameras/GameplayCameras/Source/GameplayCameras/Private/Core/ICustomCameraNodeParameterProvider.cpp
--- a/Engine/Plugins/Cameras/GameplayCameras/Source/GameplayCameras/Private/Core/ICustomCameraNodeParameterProvider.cpp
+++ b/Engine/Plugins/Cameras/GameplayCameras/Source/GameplayCameras/Private/Core/ICustomCameraNodeParameterProvider.cpp
@@ -4,6 +4,7 @@
 
 #include "Core/CameraNode.h"
 #include "Core/CameraParameters.h"
+#include "Core/CameraVariableAssets.h"
 #include "Core/CameraVariableReferences.h"
 #include "GameplayCamerasDelegates.h"
 
@@ -17,9 +18,10 @@ void FCameraNodeParameterInfos::AddBlendableParameter(
 		ECameraVariableType VariableType, 
 		const UScriptStruct* BlendableStructType,
 		const uint8* DefaultValue,
-		FCameraVariableID* OverrideVariableID)
+		FCameraVariableID* OverrideVariableID,
+		UCameraVariableAsset* OverrideVariable)
 {
-	BlendableParameters.Add({ ParameterName, VariableType, BlendableStructType, DefaultValue, OverrideVariableID });
+	BlendableParameters.Add({ ParameterName, VariableType, BlendableStructType, DefaultValue, OverrideVariableID, OverrideVariable });
 }
 
 void FCameraNodeParameterInfos::AddBlendableParameter(FCustomCameraNodeBlendableParameter& Parameter, const uint8* DefaultValue)
@@ -29,7 +31,8 @@ void FCameraNodeParameterInfos::AddBlendableParameter(FCustomCameraNodeBlendable
 			Parameter.VariableType,
 			Parameter.BlendableStructType,
 			DefaultValue,
-			Parameter.OverrideVariable ? nullptr : &Parameter.OverrideVariableID);
+			&Parameter.OverrideVariableID,
+			Parameter.OverrideVariable);
 }
 
 void FCameraNodeParameterInfos::AddDataParameter(
@@ -112,7 +115,8 @@ void FCameraNodeParameterInfos::BuildFrom(UCameraNode* InCameraNode)
 						ECameraVariableType::ValueName,\
 						nullptr,\
 						reinterpret_cast<uint8*>(&CameraParameterPtr->Value),\
-						&CameraParameterPtr->VariableID);\
+						&CameraParameterPtr->VariableID,\
+						CameraParameterPtr->Variable);\
 			}\
 			else if (StructProperty->Struct == F##ValueName##CameraVariableReference::StaticStruct())\
 			{\
@@ -122,7 +126,8 @@ void FCameraNodeParameterInfos::BuildFrom(UCameraNode* InCameraNode)
 						ECameraVariableType::ValueName,\
 						nullptr,\
 						nullptr,\
-						&VariableReferencePtr->VariableID);\
+						&VariableReferencePtr->VariableID,\
+						VariableReferencePtr->Variable);\
 			}\
 			else
 			UE_CAMERA_VARIABLE_FOR_ALL_TYPES()
@@ -225,7 +230,8 @@ void FCameraNodeParameterInfos::BuildFrom(UCameraNode* InCameraNode)
 						VariableType,
 						VariableTypeObject,
 						static_cast<uint8*>(DefaultValue),
-						BlendableID);
+						BlendableID,
+						nullptr);
 				continue;
 			}
 		}
diff --git a/Engine/Plugins/Cameras/GameplayCameras/Source/GameplayCameras/Public/Core/ICustomCameraNodeParameterProvider.h b/Engine/Plugins/Cameras/GameplayCameras/Source/GameplayCameras/Public/Core/ICustomCameraNodeParameterProvider.h
--- a/Engine/Plugins/Cameras/GameplayCameras/Source/GameplayCameras/Public/Core/ICustomCameraNodeParameterProvider.h
+++ b/Engine/Plugins/Cameras/GameplayCameras/Source/GameplayCameras/Public/Core/ICustomCameraNodeParameterProvider.h
@@ -61,7 +61,8 @@ public:
 			ECameraVariableType VariableType, 
 			const UScriptStruct* BlendableStructType,
 			const uint8* DefaultValue,
-			FCameraVariableID* OverrideVariableID);
+			FCameraVariableID* OverrideVariableID,
+			UCameraVariableAsset* OverrideVariable = nullptr);
 
 	GAMEPLAYCAMERAS_API void AddBlendableParameter(FCustomCameraNodeBlendableParameter& Parameter, const uint8* DefaultValue);
 
