Cameras: make the easing blend type a BlueprintType

#Synced-CL 49781371

diff --git a/Engine/Plugins/Cameras/GameplayCameras/Source/GameplayCameras/Public/Nodes/Blends/EasingBlendCameraNode.h b/Engine/Plugins/Cameras/GameplayCameras/Source/GameplayCameras/Public/Nodes/Blends/EasingBlendCameraNode.h
--- a/Engine/Plugins/Cameras/GameplayCameras/Source/GameplayCameras/Public/Nodes/Blends/EasingBlendCameraNode.h
+++ b/Engine/Plugins/Cameras/GameplayCameras/Source/GameplayCameras/Public/Nodes/Blends/EasingBlendCameraNode.h
@@ -12,8 +12,8 @@
  * These types are copied from, and behave in the same way as the corresponding easing
  * types from Sequencer.
  */
-UENUM()
-enum class EEasingCameraBlendType
+UENUM(BlueprintType)
+enum class EEasingCameraBlendType : uint8
 {
 	// Linear easing
 	Linear UMETA(Grouping=Linear, DisplayName="Linear"),
