Cameras: more fixes for not modifying assets during load/build unless warranted

- The code that upgrades old Camera Asset parameters to the new interface needs to reproduce the old bug that even non-visible parameters were included :(
- Non-visible parameters don't have the CPF_Edit propery flag, so also patch that when loading data in the version of the engine that didn't handle property flags in property bags

#Synced-CL 49781371

diff --git a/Engine/Plugins/Cameras/GameplayCameras/Source/GameplayCameras/Private/Build/CameraAssetBuilder.cpp b/Engine/Plugins/Cameras/GameplayCameras/Source/GameplayCameras/Private/Build/CameraAssetBuilder.cpp
--- a/Engine/Plugins/Cameras/GameplayCameras/Source/GameplayCameras/Private/Build/CameraAssetBuilder.cpp
+++ b/Engine/Plugins/Cameras/GameplayCameras/Source/GameplayCameras/Private/Build/CameraAssetBuilder.cpp
@@ -41,11 +41,8 @@ void BuildCameraAssetDefaultParameterDefinitions(UCameraAsset* CameraAsset, TArr
 			TArray<FCameraObjectInterfaceParameterDefinition>& DefinitionsForCameraRig = DefinitionsByCameraRig.Add(CameraRig);
 			for (const FCameraObjectInterfaceParameterDefinition& Definition : CameraRig->GetParameterDefinitions())
 			{
-				if (Definition.bIsVisible)
-				{
-					DefinitionsForCameraRig.Add(Definition);
-					UsedParameterNames.FindOrAdd(Definition.ParameterName).Add(CameraRig);
-				}
+				DefinitionsForCameraRig.Add(Definition);
+				UsedParameterNames.FindOrAdd(Definition.ParameterName).Add(CameraRig);
 			}
 		}
 	}
diff --git a/Engine/Plugins/Cameras/GameplayCameras/Source/GameplayCameras/Private/Build/CameraObjectInterfaceParameterBuilder.cpp b/Engine/Plugins/Cameras/GameplayCameras/Source/GameplayCameras/Private/Build/CameraObjectInterfaceParameterBuilder.cpp
--- a/Engine/Plugins/Cameras/GameplayCameras/Source/GameplayCameras/Private/Build/CameraObjectInterfaceParameterBuilder.cpp
+++ b/Engine/Plugins/Cameras/GameplayCameras/Source/GameplayCameras/Private/Build/CameraObjectInterfaceParameterBuilder.cpp
@@ -421,9 +421,11 @@ void FCameraObjectInterfaceParameterBuilder::FixUpDefaultParameterProperties(TCo
 	}
 
 	// Before UE 5.8, there was a bug with FInstancedPropertyBag not saving the PropertyFlags of its UPropertyBag
-	// struct, resulting in camera parameters losing their CPF_Interp flag. This meant that rebuilding a camera
-	// asset would create a new parameter struct that was slightly different from the loaded one (it had those
-	// flags). This made camera assets always dirty and in need of re-saving.
+	// struct, resulting in camera parameters losing their CPF_Interp flag, or still having a CPF_Edit flag even
+	// though they're supposed to not be visible/editable.
+	//
+	// This meant that rebuilding a camera asset would create a new parameter struct that was slightly different
+	// from the loaded one (it had those flags). This made camera assets always dirty and in need of re-saving.
 	//
 	// This function aims to fix that as a band-aid: it gets called on PostLoad by camera assets and re-adds the
 	// missing flags.
@@ -458,6 +460,13 @@ void FCameraObjectInterfaceParameterBuilder::FixUpDefaultParameterProperties(TCo
 			bFixedAny |= ((PropertyDesc->PropertyFlags & CPF_Interp) == 0);
 			PropertyDesc->PropertyFlags |= CPF_Interp;
 		}
+
+		// Remove CPF_Edit if the property isn't visible/editable.
+		if (!Definition.bIsVisible)
+		{
+			bFixedAny = true;
+			PropertyDesc->PropertyFlags &= ~CPF_Edit;
+		}
 	}
 
 	if (bFixedAny)
