Cameras: fix issues with camera actions started on the same frame camera rigs are activated

When a camera rig instance is pushed on the main blend stack on the same frame that a camera action is started, that action could be created on the previously active camera rig, resulting in incorrect behaviour as the action only affected the camera rig that was blending out. This change therefore starts actions in a "pending" state where we don't create the action evaluator until after the camera director has run.

To make this possible, evaluation services have a new callback for evaluating between the camera director update and the root node update.

Added new unit tests for camera actions.

#Synced-CL 49781371

diff --git a/Engine/Plugins/Cameras/GameplayCameras/Source/GameplayCameras/Private/Core/CameraEvaluationService.cpp b/Engine/Plugins/Cameras/GameplayCameras/Source/GameplayCameras/Private/Core/CameraEvaluationService.cpp
--- a/Engine/Plugins/Cameras/GameplayCameras/Source/GameplayCameras/Private/Core/CameraEvaluationService.cpp
+++ b/Engine/Plugins/Cameras/GameplayCameras/Source/GameplayCameras/Private/Core/CameraEvaluationService.cpp
@@ -25,6 +25,11 @@ void FCameraEvaluationService::PreUpdate(const FCameraEvaluationServiceUpdatePar
 	OnPreUpdate(Params, OutResult);
 }
 
+void FCameraEvaluationService::PostCameraDirectorUpdate(const FCameraEvaluationServiceUpdateParams& Params, FCameraEvaluationServiceUpdateResult& OutResult)
+{
+	OnPostCameraDirectorUpdate(Params, OutResult);
+}
+
 void FCameraEvaluationService::PostUpdate(const FCameraEvaluationServiceUpdateParams& Params, FCameraEvaluationServiceUpdateResult& OutResult)
 {
 	OnPostUpdate(Params, OutResult);
diff --git a/Engine/Plugins/Cameras/GameplayCameras/Source/GameplayCameras/Private/Core/CameraSystemEvaluator.cpp b/Engine/Plugins/Cameras/GameplayCameras/Source/GameplayCameras/Private/Core/CameraSystemEvaluator.cpp
--- a/Engine/Plugins/Cameras/GameplayCameras/Source/GameplayCameras/Private/Core/CameraSystemEvaluator.cpp
+++ b/Engine/Plugins/Cameras/GameplayCameras/Source/GameplayCameras/Private/Core/CameraSystemEvaluator.cpp
@@ -241,7 +241,7 @@ void FCameraSystemEvaluator::UpdateImpl(float DeltaTime, ECameraNodeEvaluationTy
 	RootNodeResult.ContextDataTable.AutoResetValues();
 
 	// Pre-update all services.
-	PreUpdateServices(DeltaTime, EvaluationType, ECameraEvaluationServiceFlags::None);
+	PreUpdateServices(DeltaTime, EvaluationType);
 
 	// Get the active evaluation context.
 	TSharedPtr<FCameraEvaluationContext> ActiveContext = ContextStack.GetActiveContext();
@@ -266,6 +266,8 @@ void FCameraSystemEvaluator::UpdateImpl(float DeltaTime, ECameraNodeEvaluationTy
 		UpdateCameraDirector(DeltaTime, ActiveDirectorEvaluator);
 	}
 
+	PostCameraDirectorUpdateServices(DeltaTime, EvaluationType);
+
 	// Run the camera node tree.
 	{
 		FCameraNodeEvaluationParams NodeParams;
@@ -281,7 +283,7 @@ void FCameraSystemEvaluator::UpdateImpl(float DeltaTime, ECameraNodeEvaluationTy
 	}
 
 	// Post-update all services.
-	PostUpdateServices(DeltaTime, EvaluationType, ECameraEvaluationServiceFlags::None);
+	PostUpdateServices(DeltaTime, EvaluationType);
 
 	// Harvest the result.
 	PreVisualResult.Reset(RootEvaluator->GetPreVisualLayerResult());
@@ -411,7 +413,7 @@ void FCameraSystemEvaluator::GetCombinedCameraRigRequest(TConstArrayView<FCamera
 	OutCombinedRequest.bForceActivateDeactivate = bAnyForceActivationDeactivation;
 }
 
-void FCameraSystemEvaluator::PreUpdateServices(float DeltaTime, ECameraNodeEvaluationType EvaluationType, ECameraEvaluationServiceFlags ExtraFlags)
+void FCameraSystemEvaluator::PreUpdateServices(float DeltaTime, ECameraNodeEvaluationType EvaluationType)
 {
 	FCameraEvaluationServiceUpdateParams ServiceUpdateParams;
 	ServiceUpdateParams.Evaluator = this;
@@ -422,14 +424,32 @@ void FCameraSystemEvaluator::PreUpdateServices(float DeltaTime, ECameraNodeEvalu
 
 	for (TSharedPtr<FCameraEvaluationService> EvaluationService : EvaluationServices)
 	{
-		if (EvaluationService->HasAllEvaluationServiceFlags(ECameraEvaluationServiceFlags::NeedsPreUpdate | ExtraFlags))
+		if (EvaluationService->HasAllEvaluationServiceFlags(ECameraEvaluationServiceFlags::NeedsPreUpdate))
 		{
 			EvaluationService->PreUpdate(ServiceUpdateParams, ServiceUpdateResult);
 		}
 	}
 }
 
-void FCameraSystemEvaluator::PostUpdateServices(float DeltaTime, ECameraNodeEvaluationType EvaluationType, ECameraEvaluationServiceFlags ExtraFlags)
+void FCameraSystemEvaluator::PostCameraDirectorUpdateServices(float DeltaTime, ECameraNodeEvaluationType EvaluationType)
+{
+	FCameraEvaluationServiceUpdateParams ServiceUpdateParams;
+	ServiceUpdateParams.Evaluator = this;
+	ServiceUpdateParams.DeltaTime = DeltaTime;
+	ServiceUpdateParams.EvaluationType = EvaluationType;
+
+	FCameraEvaluationServiceUpdateResult ServiceUpdateResult(RootNodeResult);
+
+	for (TSharedPtr<FCameraEvaluationService> EvaluationService : EvaluationServices)
+	{
+		if (EvaluationService->HasAllEvaluationServiceFlags(ECameraEvaluationServiceFlags::NeedsPostCameraDirectorUpdate))
+		{
+			EvaluationService->PostCameraDirectorUpdate(ServiceUpdateParams, ServiceUpdateResult);
+		}
+	}
+}
+
+void FCameraSystemEvaluator::PostUpdateServices(float DeltaTime, ECameraNodeEvaluationType EvaluationType)
 {
 	FCameraEvaluationServiceUpdateParams ServiceUpdateParams;
 	ServiceUpdateParams.Evaluator = this;
@@ -440,7 +460,7 @@ void FCameraSystemEvaluator::PostUpdateServices(float DeltaTime, ECameraNodeEval
 
 	for (TSharedPtr<FCameraEvaluationService> EvaluationService : EvaluationServices)
 	{
-		if (EvaluationService->HasAllEvaluationServiceFlags(ECameraEvaluationServiceFlags::NeedsPostUpdate | ExtraFlags))
+		if (EvaluationService->HasAllEvaluationServiceFlags(ECameraEvaluationServiceFlags::NeedsPostUpdate))
 		{
 			EvaluationService->PostUpdate(ServiceUpdateParams, ServiceUpdateResult);
 		}
diff --git a/Engine/Plugins/Cameras/GameplayCameras/Source/GameplayCameras/Private/Services/CameraActionService.cpp b/Engine/Plugins/Cameras/GameplayCameras/Source/GameplayCameras/Private/Services/CameraActionService.cpp
--- a/Engine/Plugins/Cameras/GameplayCameras/Source/GameplayCameras/Private/Services/CameraActionService.cpp
+++ b/Engine/Plugins/Cameras/GameplayCameras/Source/GameplayCameras/Private/Services/CameraActionService.cpp
@@ -24,9 +24,7 @@ UE_DEFINE_CAMERA_EVALUATION_SERVICE(FCameraActionService)
 
 FCameraActionService::~FCameraActionService()
 {
-	// If we are destroying the owning evaluator while a cloning camera action is running,
-	// we need to clean-up our delegate for root node events.
-	if (RootNodeEvaluator && NumCloningActions > 0)
+	if (RootNodeEvaluator)
 	{
 		RootNodeEvaluator->OnCameraRigEvent().RemoveAll(this);
 	}
@@ -34,9 +32,54 @@ FCameraActionService::~FCameraActionService()
 
 void FCameraActionService::OnInitialize(const FCameraEvaluationServiceInitializeParams& Params)
 {
-	SetEvaluationServiceFlags(ECameraEvaluationServiceFlags::NeedsPostUpdate);
+	SetEvaluationServiceFlags(
+			ECameraEvaluationServiceFlags::NeedsPostCameraDirectorUpdate |
+			ECameraEvaluationServiceFlags::NeedsPostUpdate);
 
 	RootNodeEvaluator = Params.Evaluator->GetRootNodeEvaluator();
+	if (ensure(RootNodeEvaluator))
+	{
+		RootNodeEvaluator->OnCameraRigEvent().AddRaw(this, &FCameraActionService::OnRootNodeCameraRigEvent);
+	}
+}
+
+void FCameraActionService::OnPostCameraDirectorUpdate(const FCameraEvaluationServiceUpdateParams& Params, FCameraEvaluationServiceUpdateResult& OutResult)
+{
+	for (FActionSetInfo& ActionSet : ActionSets)
+	{
+		if (!ensure(ActionSet.Data))
+		{
+			continue;
+		}
+
+		if (ActionSet.bIsPending)
+		{
+			// This action is pending and its first evaluator needs to be created on the active camera rig instance's
+			// action scope.
+			TSharedPtr<FCameraActionScope> ActionScope = RootNodeEvaluator->GetActiveCameraRigActionScope(true);
+			if (ensure(ActionScope))
+			{
+				FCameraActionInstanceID NewInstanceID = ActionScope->StartAction(ActionSet.Data);
+				if (ensure(NewInstanceID.IsValid()))
+				{
+					ActionSet.Actions.Add(FActionInfo{ ActionScope, NewInstanceID });
+				}
+			}
+			ActionSet.bIsPending = false;
+		}
+		else if (bHasAnyNewActiveCameraRigs && ActionSet.Data->bPropagateToNewCameraRigs)
+		{
+			// This action is already running but a new camera rig instance has been activated, and the action wants
+			// to propagate to it, so we clone its evaluator onto the new action scope.
+			TSharedPtr<FCameraActionScope> ActionScope = RootNodeEvaluator->GetActiveCameraRigActionScope(true);
+			if (ensure(ActionScope))
+			{
+				CloneAction(ActionSet, ActionScope.ToSharedRef());
+			}
+		}
+	}
+
+	bHasAnyNewActiveCameraRigs = false;
 }
 
 void FCameraActionService::OnPostUpdate(const FCameraEvaluationServiceUpdateParams& Params, FCameraEvaluationServiceUpdateResult& OutResult)
@@ -47,35 +90,37 @@ void FCameraActionService::OnPostUpdate(const FCameraEvaluationServiceUpdatePara
 	}
 }
 
-FCameraActionInstanceID FCameraActionService::StartAction(const UCameraAction* CameraAction)
+void FCameraActionService::OnTeardown(const FCameraEvaluationServiceTeardownParams& Params)
 {
-	if (!CameraAction)
+	if (ensure(RootNodeEvaluator))
 	{
-		return FCameraActionInstanceID();
+		RootNodeEvaluator->OnCameraRigEvent().RemoveAll(this);
 	}
-	
-	TSharedPtr<FCameraActionScope> ActionScope = RootNodeEvaluator->GetActiveCameraRigActionScope(true);
-	if (!ensure(ActionScope))
+
+	RootNodeEvaluator = nullptr;
+	ActionSets.Reset();
+}
+
+void FCameraActionService::OnRootNodeCameraRigEvent(const FRootCameraNodeCameraRigEvent& InEvent)
+{
+	if (InEvent.EventType == ERootCameraNodeCameraRigEventType::Activated && InEvent.EventLayer == ECameraRigLayer::Main)
 	{
-		return FCameraActionInstanceID();
+		bHasAnyNewActiveCameraRigs = true;
 	}
-	
-	FCameraActionInstanceID NewInstanceID = ActionScope->StartAction(CameraAction);
-	if (!ensure(NewInstanceID.IsValid()))
+}
+
+FCameraActionInstanceID FCameraActionService::StartAction(const UCameraAction* CameraAction)
+{
+	if (!CameraAction)
 	{
 		return FCameraActionInstanceID();
 	}
-
+	
 	FActionSetInfo& NewActionSet = ActionSets.Emplace_GetRef();
 	NewActionSet.Data = CameraAction;
-	NewActionSet.Actions.Add(FActionInfo{ ActionScope, NewInstanceID });
 	NewActionSet.InstanceID = FCameraActionInstanceID(NextActionSetID++);
+	NewActionSet.bIsPending = true;
 
-	if (CameraAction->bPropagateToNewCameraRigs)
-	{
-		AddCloningAction();
-	}
-	
 	return NewActionSet.InstanceID;
 }
 
@@ -109,10 +154,6 @@ bool FCameraActionService::StopAction(const FCameraActionInstanceID InInstanceID
 				ensure(bStopped);
 			}
 		}
-		if (ActionSet.Data->bPropagateToNewCameraRigs)
-		{
-			RemoveCloningAction();
-		}
 		ActionSets.RemoveAt(Index);
 		return true;
 	}
@@ -136,57 +177,24 @@ bool FCameraActionService::StopAllActionsOfClass(TSubclassOf<UCameraAction> InAc
 					bAnyStopped = true;
 				}
 			}
-			if (ActionSet.Data->bPropagateToNewCameraRigs)
-			{
-				RemoveCloningAction();
-			}
 			It.RemoveCurrent();
 		}
 	}
 	return bAnyStopped;
 }
 
-void FCameraActionService::OnRootCameraNodeCameraRigEvent(const FRootCameraNodeCameraRigEvent& InEvent)
+int32 FCameraActionService::GetNumScopeActions(const FCameraActionInstanceID InInstanceID) const
 {
-	if (InEvent.EventType == ERootCameraNodeCameraRigEventType::Activated && InEvent.EventLayer == ECameraRigLayer::Main)
-	{
-		TSharedPtr<FCameraActionScope> ActionScope = RootNodeEvaluator->GetActiveCameraRigActionScope(true);
-		if (ensure(ActionScope))
-		{
-			int32 NumCloned = 0;
-			for (FActionSetInfo& ActionSet : ActionSets)
+	const int32 Index = ActionSets.IndexOfByPredicate([InInstanceID](const FActionSetInfo& Item)
 			{
-				if (ensure(ActionSet.Data) && ActionSet.Data->bPropagateToNewCameraRigs)
-				{
-					CloneAction(ActionSet, ActionScope.ToSharedRef());
-					++NumCloned;
-				}
-			}
-			ensureMsgf(
-					NumCloned > 0, 
-					TEXT("No actions found to clone: we are still registered with root camera rig events for nothing"));
-		}
-	}
-}
-
-void FCameraActionService::AddCloningAction()
-{
-	const bool bRegisterEvent = (NumCloningActions == 0);
-	++NumCloningActions;
-	if (bRegisterEvent)
-	{
-		RootNodeEvaluator->OnCameraRigEvent().AddRaw(this, &FCameraActionService::OnRootCameraNodeCameraRigEvent);
-	}
-}
-
-void FCameraActionService::RemoveCloningAction()
-{
-	ensure(NumCloningActions > 0);
-	--NumCloningActions;
-	if (NumCloningActions == 0)
+				return Item.InstanceID == InInstanceID;
+			});
+	if (Index != INDEX_NONE)
 	{
-		RootNodeEvaluator->OnCameraRigEvent().RemoveAll(this);
+		const FActionSetInfo& ActionSet = ActionSets[Index];
+		return ActionSet.Actions.Num();
 	}
+	return 0;
 }
 
 void FCameraActionService::CloneAction(FActionSetInfo& ActionSet, TSharedRef<FCameraActionScope> NewActionScope)
@@ -232,6 +240,9 @@ void FCameraActionService::CleanUpActions()
 	for (auto SetIt = ActionSets.CreateIterator(); SetIt; ++SetIt)
 	{
 		FActionSetInfo& ActionSet(*SetIt);
+
+		ensure(!ActionSet.bIsPending);
+
 		for (auto ActionIt = ActionSet.Actions.CreateIterator(); ActionIt; ++ActionIt)
 		{
 			// Remove finished actions, whether they finished on their own, or were ended by their action scope 
@@ -259,10 +270,6 @@ void FCameraActionService::CleanUpActions()
 		}
 		if (ActionSet.Actions.IsEmpty())
 		{
-			if (ActionSet.Data->bPropagateToNewCameraRigs)
-			{
-				RemoveCloningAction();
-			}
 #if UE_VERSION_NEWER_THAN_OR_EQUAL(5,8,0)
 			SetIt.RemoveCurrent(EAllowShrinking::No);
 #else
diff --git a/Engine/Plugins/Cameras/GameplayCameras/Source/GameplayCameras/Private/Tests/CameraActionTests.cpp b/Engine/Plugins/Cameras/GameplayCameras/Source/GameplayCameras/Private/Tests/CameraActionTests.cpp
new file mode 100644
--- /dev/null
+++ b/Engine/Plugins/Cameras/GameplayCameras/Source/GameplayCameras/Private/Tests/CameraActionTests.cpp
@@ -0,0 +1,187 @@
+// Copyright Epic Games, Inc. All Rights Reserved.
+
+#include "Core/CameraSystemEvaluator.h"
+#include "Core/RootCameraNode.h"
+#include "Misc/AutomationTest.h"
+#include "Nodes/Blends/LinearBlendCameraNode.h"
+#include "Nodes/Common/OffsetCameraNode.h"
+#include "Services/CameraActionService.h"
+#include "Tests/GameplayCamerasTestBuilder.h"
+#include "Tests/GameplayCamerasTestObjects.h"
+
+namespace UE::Cameras::Test
+{
+	struct FCameraActionTestData
+	{
+		TSharedPtr<FCameraSystemEvaluator> Evaluator;
+		TSharedPtr<FCameraActionService> ActionService;
+		FRootCameraNodeEvaluator* RootEvaluator = nullptr;
+
+		TSharedPtr<FCameraEvaluationContext> TestEvaluationContext;
+		UCameraRigAsset* TestCameraRig = nullptr;
+
+		void UpdateEvaluator(float DeltaTime = 0.3f);
+		void ForceReactivateTestCameraRig(const UCameraRigTransition* TransitionOverride = nullptr);
+
+		TSharedPtr<FCameraActionScope> GetActiveActionScope();
+		FCameraActionInstanceID StartAction(const UCameraAction* Action);
+	};
+
+	FCameraActionTestData CreateCameraActionTestData()
+	{
+		FCameraActionTestData TestData;
+
+		// Create the evaluator, and a test evaluation context.
+		TestData.Evaluator = MakeShared<FCameraSystemEvaluator>();
+		TestData.Evaluator->Initialize();
+
+		TestData.TestEvaluationContext = FCameraEvaluationContextTestBuilder()
+			.CreateCameraRig(TEXT("TestRig"))
+				.MakeRootNode<UOffsetCameraNode>()
+					.Done()
+				.Pin(TestData.TestCameraRig)
+				.Done()
+			.MakeSingleDirector()
+				.Setup([](USingleCameraDirector* Director, FNamedObjectRegistry* NamedObjectRegistry)
+					{
+						Director->CameraRig = NamedObjectRegistry->Get<UCameraRigAsset>(TEXT("TestRig"));
+					})
+				.Done()
+			.BuildCameraAsset()
+			.Get();
+		ensure(TestData.TestCameraRig);
+
+		// Activate the test evaluation context.
+		TestData.Evaluator->PushEvaluationContext(TestData.TestEvaluationContext.ToSharedRef());
+
+		// Grab the action service and root evaluator.
+		TestData.ActionService = TestData.Evaluator->FindEvaluationService<FCameraActionService>();
+		ensure(TestData.ActionService.IsValid());
+
+		TestData.RootEvaluator = TestData.Evaluator->GetRootNodeEvaluator();
+		ensure(TestData.RootEvaluator);
+
+		// Update the evaluator once so we get usable setup.
+		TestData.UpdateEvaluator();
+
+		return TestData;
+	}
+
+	void FCameraActionTestData::UpdateEvaluator(float DeltaTime)
+	{
+		FCameraSystemEvaluationParams Params;
+		Params.DeltaTime = DeltaTime;
+		Evaluator->Update(Params);
+	}
+
+	void FCameraActionTestData::ForceReactivateTestCameraRig(const UCameraRigTransition* TransitionOverride)
+	{
+		FActivateCameraRigParams NewRigParams;
+		NewRigParams.EvaluationContext = TestEvaluationContext;
+		NewRigParams.CameraRig = TestCameraRig;
+		NewRigParams.TransitionOverride = TransitionOverride;
+		NewRigParams.bForceActivate = true;
+		RootEvaluator->ActivateCameraRig(NewRigParams);
+	}
+
+	TSharedPtr<FCameraActionScope> FCameraActionTestData::GetActiveActionScope()
+	{
+		return RootEvaluator->GetActiveCameraRigActionScope(false);
+	}
+
+	FCameraActionInstanceID FCameraActionTestData::StartAction(const UCameraAction* Action)
+	{
+		return ActionService->StartAction(Action);
+	}
+
+}  // namespace UE::Cameras::Test
+
+IMPLEMENT_SIMPLE_AUTOMATION_TEST(FCameraActionSimpleTest, "System.Engine.GameplayCameras.CameraAction.Simple", EAutomationTestFlags::EditorContext | EAutomationTestFlags::EngineFilter)
+bool FCameraActionSimpleTest::RunTest(const FString& Parameters)
+{
+	using namespace UE::Cameras;
+	using namespace UE::Cameras::Test;
+
+	FCameraActionTestData TestData = CreateCameraActionTestData();
+
+	TSharedPtr<FCameraActionScope> ActiveActionScope = TestData.GetActiveActionScope();
+	UTEST_NULL("No initial action scope", ActiveActionScope.Get());
+
+	// Push our test action.
+	UUpdateTrackerCameraAction* NewAction = NewObject<UUpdateTrackerCameraAction>();
+	NewAction->NumUpdates = 2;
+	FCameraActionInstanceID NewInstanceID = TestData.StartAction(NewAction);
+
+	// Update twice... the action should be gone after the second time.
+	TestData.UpdateEvaluator();
+	UTEST_TRUE("Action is running", TestData.ActionService->IsActionRunning(NewInstanceID));
+
+	TestData.UpdateEvaluator();
+	UTEST_FALSE("Action has finished", TestData.ActionService->IsActionRunning(NewInstanceID));
+
+	return true;
+}
+
+IMPLEMENT_SIMPLE_AUTOMATION_TEST(FCameraActionCloningTest, "System.Engine.GameplayCameras.CameraAction.Cloning", EAutomationTestFlags::EditorContext | EAutomationTestFlags::EngineFilter)
+bool FCameraActionCloningTest::RunTest(const FString& Parameters)
+{
+	using namespace UE::Cameras;
+	using namespace UE::Cameras::Test;
+
+	FCameraActionTestData TestData = CreateCameraActionTestData();
+
+	// Push our test action as clonable.
+	UUpdateTrackerCameraAction* NewAction = NewObject<UUpdateTrackerCameraAction>();
+	NewAction->bPropagateToNewCameraRigs = true;
+	NewAction->NumUpdates = 3;
+	FCameraActionInstanceID NewInstanceID = TestData.StartAction(NewAction);
+
+	// Update once, the camera action should be running.
+	TestData.UpdateEvaluator();
+	UTEST_TRUE("Action is running", TestData.ActionService->IsActionRunning(NewInstanceID));
+	UTEST_EQUAL("One action active", TestData.ActionService->GetNumScopeActions(NewInstanceID), 1);
+
+	// Force activate a new rig, with a transition blend to make sure both rigs will run for a few more frames.
+	UCameraRigTransition* NewTransition = NewObject<UCameraRigTransition>();
+	NewTransition->Blend = NewObject<ULinearBlendCameraNode>();
+	TestData.ForceReactivateTestCameraRig(NewTransition);
+	TestData.UpdateEvaluator();
+
+	// The action should have been cloned over, with both instances left with 2 updates.
+	UTEST_TRUE("Action is running", TestData.ActionService->IsActionRunning(NewInstanceID));
+	UTEST_EQUAL("Two actions active", TestData.ActionService->GetNumScopeActions(NewInstanceID), 2);
+
+	TestData.UpdateEvaluator();
+	UTEST_FALSE("Action has finished", TestData.ActionService->IsActionRunning(NewInstanceID));
+
+	return true;
+}
+
+IMPLEMENT_SIMPLE_AUTOMATION_TEST(FCameraActionPendingTest, "System.Engine.GameplayCameras.CameraAction.StartAndActivateOnSameFrame", EAutomationTestFlags::EditorContext | EAutomationTestFlags::EngineFilter)
+bool FCameraActionPendingTest::RunTest(const FString& Parameters)
+{
+	using namespace UE::Cameras;
+	using namespace UE::Cameras::Test;
+
+	FCameraActionTestData TestData = CreateCameraActionTestData();
+
+	// In the same frame, start an action and push a new rig instance.
+	UUpdateTrackerCameraAction* NewAction = NewObject<UUpdateTrackerCameraAction>();
+	NewAction->NumUpdates = 2;
+	FCameraActionInstanceID NewInstanceID = TestData.StartAction(NewAction);
+
+	TestData.ForceReactivateTestCameraRig();
+	TestData.UpdateEvaluator();
+
+	// The camera action should be running, but only in the newest camera rig's scope.
+	UTEST_TRUE("Action is running", TestData.ActionService->IsActionRunning(NewInstanceID));
+	UTEST_EQUAL("One action active", TestData.ActionService->GetNumScopeActions(NewInstanceID), 1);
+
+	// Run again, it should be done.
+	TestData.UpdateEvaluator();
+	UTEST_FALSE("Action has finished", TestData.ActionService->IsActionRunning(NewInstanceID));
+	UTEST_EQUAL("No more action active", TestData.ActionService->GetNumScopeActions(NewInstanceID), 0);
+
+	return true;
+}
+
diff --git a/Engine/Plugins/Cameras/GameplayCameras/Source/GameplayCameras/Private/Tests/GameplayCamerasTestObjects.cpp b/Engine/Plugins/Cameras/GameplayCameras/Source/GameplayCameras/Private/Tests/GameplayCamerasTestObjects.cpp
--- a/Engine/Plugins/Cameras/GameplayCameras/Source/GameplayCameras/Private/Tests/GameplayCamerasTestObjects.cpp
+++ b/Engine/Plugins/Cameras/GameplayCameras/Source/GameplayCameras/Private/Tests/GameplayCamerasTestObjects.cpp
@@ -37,7 +37,27 @@ void FFixedTextCameraDirectorEvaluator::OnRun(const FCameraDirectorEvaluationPar
 	}
 }
 
-}  // namespace UE::Cameras::Tests
+UE_DEFINE_CAMERA_ACTION_EVALUATOR(FUpdateTrackerCameraActionEvaluator)
+
+void FUpdateTrackerCameraActionEvaluator::OnPreScopeRun(const FCameraActionEvaluationParams& Params, FCameraActionEvaluationResult& OutResult)
+{
+	const UUpdateTrackerCameraAction* Action = GetCameraActionAs<UUpdateTrackerCameraAction>();
+	if (ensure(Action))
+	{
+		++NumUpdatesReceived;
+		if (NumUpdatesReceived >= Action->NumUpdates)
+		{
+			OutResult.bIsActionFinished = true;
+		}
+	}
+}
+
+void FUpdateTrackerCameraActionEvaluator::OnSerialize(const FCameraActionEvaluatorSerializeParams& Params, FArchive& Ar)
+{
+	Ar << NumUpdatesReceived;
+}
+
+}  // namespace UE::Cameras::Test
 
 FCameraNodeEvaluatorPtr UUpdateTrackerCameraNode::OnBuildEvaluator(FCameraNodeEvaluatorBuilder& Builder) const
 {
@@ -81,3 +101,9 @@ void UFixedTestCameraDirector::OnGatherRigUsageInfo(FCameraDirectorRigUsageInfo&
 	UsageInfo.CameraRigs.Append(CameraRigs);
 }
 
+FCameraActionEvaluatorPtr UUpdateTrackerCameraAction::OnBuildEvaluator(FCameraActionEvaluatorBuilder& Builder) const
+{
+	using namespace UE::Cameras::Test;
+	return Builder.BuildEvaluator<FUpdateTrackerCameraActionEvaluator>();
+}
+
diff --git a/Engine/Plugins/Cameras/GameplayCameras/Source/GameplayCameras/Private/Tests/GameplayCamerasTestObjects.h b/Engine/Plugins/Cameras/GameplayCameras/Source/GameplayCameras/Private/Tests/GameplayCamerasTestObjects.h
--- a/Engine/Plugins/Cameras/GameplayCameras/Source/GameplayCameras/Private/Tests/GameplayCamerasTestObjects.h
+++ b/Engine/Plugins/Cameras/GameplayCameras/Source/GameplayCameras/Private/Tests/GameplayCamerasTestObjects.h
@@ -8,6 +8,8 @@
 #include "Core/CameraNodeEvaluator.h"
 #include "Core/CameraParameters.h"
 #include "Nodes/CameraNodeTypes.h"
+#include "Services/CameraAction.h"
+#include "Services/CameraActionEvaluator.h"
 
 #include "GameplayCamerasTestObjects.generated.h"
 
@@ -53,7 +55,22 @@ private:
 	int32 ActiveIndex = INDEX_NONE;
 };
 
-}  // namespace UE::Cameras::Tests
+class FUpdateTrackerCameraActionEvaluator : public FCameraActionEvaluator
+{
+	UE_DECLARE_CAMERA_ACTION_EVALUATOR(GAMEPLAYCAMERAS_API, FUpdateTrackerCameraActionEvaluator)
+
+protected:
+
+	// FCameraActionEvaluator interface.
+	virtual void OnPreScopeRun(const FCameraActionEvaluationParams& Params, FCameraActionEvaluationResult& OutResult) override;
+	virtual void OnSerialize(const FCameraActionEvaluatorSerializeParams& Params, FArchive& Ar) override;
+
+private:
+
+	int32 NumUpdatesReceived = 0;
+};
+
+}  // namespace UE::Cameras::Test
 
 UCLASS(MinimalAPI, Hidden)
 class UUpdateTrackerCameraNode : public UCameraNode
@@ -100,3 +117,19 @@ private:
 	TArray<FName> CameraRigNames;
 };
 
+UCLASS(MinimalAPI, Hidden)
+class UUpdateTrackerCameraAction : public UCameraAction
+{
+	GENERATED_BODY()
+
+public:
+
+	UPROPERTY()
+	int32 NumUpdates = 1;
+
+protected:
+
+	// UCameraAction interface.
+	virtual FCameraActionEvaluatorPtr OnBuildEvaluator(FCameraActionEvaluatorBuilder& Builder) const;
+};
+
diff --git a/Engine/Plugins/Cameras/GameplayCameras/Source/GameplayCameras/Public/Core/CameraEvaluationService.h b/Engine/Plugins/Cameras/GameplayCameras/Source/GameplayCameras/Public/Core/CameraEvaluationService.h
--- a/Engine/Plugins/Cameras/GameplayCameras/Source/GameplayCameras/Public/Core/CameraEvaluationService.h
+++ b/Engine/Plugins/Cameras/GameplayCameras/Source/GameplayCameras/Public/Core/CameraEvaluationService.h
@@ -22,10 +22,11 @@ enum class ECameraEvaluationServiceFlags
 {
 	None = 0,
 	NeedsPreUpdate = 1 << 0,
-	NeedsPostUpdate = 1 << 1,
-	NeedsRootCameraNodeEvents = 1 << 2,
+	NeedsPostCameraDirectorUpdate = 1 << 1,
+	NeedsPostUpdate = 1 << 2,
+	NeedsRootCameraNodeEvents = 1 << 3,
 
-	Default = NeedsPreUpdate | NeedsPostUpdate | NeedsRootCameraNodeEvents
+	Default = NeedsPreUpdate | NeedsPostCameraDirectorUpdate | NeedsPostUpdate | NeedsRootCameraNodeEvents
 };
 ENUM_CLASS_FLAGS(ECameraEvaluationServiceFlags);
 
@@ -81,6 +82,8 @@ public:
 	void Initialize(const FCameraEvaluationServiceInitializeParams& Params);
 	/** Runs at the start of the camera system update. */
 	void PreUpdate(const FCameraEvaluationServiceUpdateParams& Params, FCameraEvaluationServiceUpdateResult& OutResult);
+	/** Runs after the camera director update. */
+	void PostCameraDirectorUpdate(const FCameraEvaluationServiceUpdateParams& Params, FCameraEvaluationServiceUpdateResult& OutResult);
 	/** Runs at the end of the camera system update. */
 	void PostUpdate(const FCameraEvaluationServiceUpdateParams& Params, FCameraEvaluationServiceUpdateResult& OutResult);
 	/** Tears down the evaluation service. */
@@ -109,6 +112,8 @@ protected:
 	virtual void OnInitialize(const FCameraEvaluationServiceInitializeParams& Params) {}
 	/** Runs at the start of the camera system update. */
 	virtual void OnPreUpdate(const FCameraEvaluationServiceUpdateParams& Params, FCameraEvaluationServiceUpdateResult& OutResult) {}
+	/** Runs after the camera director update. */
+	virtual void OnPostCameraDirectorUpdate(const FCameraEvaluationServiceUpdateParams& Params, FCameraEvaluationServiceUpdateResult& OutResult) {}
 	/** Runs at the end of the camera system update. */
 	virtual void OnPostUpdate(const FCameraEvaluationServiceUpdateParams& Params, FCameraEvaluationServiceUpdateResult& OutResult) {}
 	/** Tears down the evaluation service. */
diff --git a/Engine/Plugins/Cameras/GameplayCameras/Source/GameplayCameras/Public/Core/CameraSystemEvaluator.h b/Engine/Plugins/Cameras/GameplayCameras/Source/GameplayCameras/Public/Core/CameraSystemEvaluator.h
--- a/Engine/Plugins/Cameras/GameplayCameras/Source/GameplayCameras/Public/Core/CameraSystemEvaluator.h
+++ b/Engine/Plugins/Cameras/GameplayCameras/Source/GameplayCameras/Public/Core/CameraSystemEvaluator.h
@@ -301,8 +301,9 @@ private:
 	void UpdateCameraDirector(float DeltaTime, FCameraDirectorEvaluator* CameraDirectorEvaluator);
 	void GetCombinedCameraRigRequest(TConstArrayView<FCameraRigActivationDeactivationRequest> Requests, FCameraRigActivationDeactivationRequest& OutCombinedRequest);
 
-	void PreUpdateServices(float DeltaTime, ECameraNodeEvaluationType EvaluationType, ECameraEvaluationServiceFlags ExtraFlags);
-	void PostUpdateServices(float DeltaTime, ECameraNodeEvaluationType EvaluationType, ECameraEvaluationServiceFlags ExtraFlags);
+	void PreUpdateServices(float DeltaTime, ECameraNodeEvaluationType EvaluationType);
+	void PostCameraDirectorUpdateServices(float DeltaTime, ECameraNodeEvaluationType EvaluationType);
+	void PostUpdateServices(float DeltaTime, ECameraNodeEvaluationType EvaluationType);
 
 	void NotifyRootCameraNodeEvent(const FRootCameraNodeCameraRigEvent& InEvent);
 
diff --git a/Engine/Plugins/Cameras/GameplayCameras/Source/GameplayCameras/Public/Services/CameraActionService.h b/Engine/Plugins/Cameras/GameplayCameras/Source/GameplayCameras/Public/Services/CameraActionService.h
--- a/Engine/Plugins/Cameras/GameplayCameras/Source/GameplayCameras/Public/Services/CameraActionService.h
+++ b/Engine/Plugins/Cameras/GameplayCameras/Source/GameplayCameras/Public/Services/CameraActionService.h
@@ -46,11 +46,26 @@ public:
 	/** Stops all camera action instances of a given class. */
 	bool StopAllActionsOfClass(TSubclassOf<UCameraAction> InActionClass);
 
+public:
+
+	/**
+	 * Returns how many actual actions are running for the given instance ID.
+	 *
+	 * This returns 1 for a running non-propagating action.
+	 * For a propagating action, this returns the number of running camera rig instances that the action is currently
+	 * affecting, i.e. the number of running clones.
+	 * If an action has been started but the camera system hasn't updated yet, it returns 0 since action evaluators
+	 * are created only on the first update.
+	 */
+	int32 GetNumScopeActions(const FCameraActionInstanceID InInstanceID) const;
+
 protected:
 
 	// FCameraEvaluationService interface.
 	virtual void OnInitialize(const FCameraEvaluationServiceInitializeParams& Params) override;
+	virtual void OnPostCameraDirectorUpdate(const FCameraEvaluationServiceUpdateParams& Params, FCameraEvaluationServiceUpdateResult& OutResult) override;
 	virtual void OnPostUpdate(const FCameraEvaluationServiceUpdateParams& Params, FCameraEvaluationServiceUpdateResult& OutResult) override;
+	virtual void OnTeardown(const FCameraEvaluationServiceTeardownParams& Params) override;
 #if UE_GAMEPLAY_CAMERAS_DEBUG
 	virtual void OnBuildDebugBlocks(const FCameraDebugBlockBuildParams& Params, FCameraDebugBlockBuilder& Builder) override;
 #endif  // UE_GAMEPLAY_CAMERAS_DEBUG
@@ -59,10 +74,8 @@ private:
 
 	struct FActionSetInfo;
 
-	void OnRootCameraNodeCameraRigEvent(const FRootCameraNodeCameraRigEvent& InEvent);
+	void OnRootNodeCameraRigEvent(const FRootCameraNodeCameraRigEvent& InEvent);
 
-	void AddCloningAction();
-	void RemoveCloningAction();
 	void CloneAction(FActionSetInfo& ActionSet, TSharedRef<FCameraActionScope> NewActionScope);
 
 	void CleanUpActions();
@@ -80,6 +93,7 @@ private:
 		TObjectPtr<const UCameraAction> Data;
 		TArray<FActionInfo, TInlineAllocator<4>> Actions;
 		FCameraActionInstanceID InstanceID;
+		bool bIsPending = false;
 	};
 
 	FRootCameraNodeEvaluator* RootNodeEvaluator = nullptr;
@@ -87,7 +101,8 @@ private:
 	TArray<FActionSetInfo> ActionSets;
 
 	uint32 NextActionSetID = 0;
-	uint8 NumCloningActions = 0;
+
+	bool bHasAnyNewActiveCameraRigs = false;
 
 #if UE_GAMEPLAY_CAMERAS_DEBUG
 	friend class FCameraActionServiceDebugBlock;
