Cameras: remove extra root node event API

Reimplement the evaluation service's `OnRootCameraNodeEvent()` API in terms of the existing event that exists on the root node evaluator. This removes a bit of core API that wasn't needed.

#Synced-CL 49781371

diff --git a/Engine/Plugins/Cameras/GameplayCameras/Source/GameplayCameras/Private/Core/CameraEvaluationService.cpp b/Engine/Plugins/Cameras/GameplayCameras/Source/GameplayCameras/Private/Core/CameraEvaluationService.cpp
--- a/Engine/Plugins/Cameras/GameplayCameras/Source/GameplayCameras/Private/Core/CameraEvaluationService.cpp
+++ b/Engine/Plugins/Cameras/GameplayCameras/Source/GameplayCameras/Private/Core/CameraEvaluationService.cpp
@@ -2,6 +2,9 @@
 
 #include "Core/CameraEvaluationService.h"
 
+#include "Core/CameraSystemEvaluator.h"
+#include "Core/RootCameraNode.h"
+
 namespace UE::Cameras
 {
 
@@ -13,11 +16,36 @@ FCameraEvaluationService::FCameraEvaluationService()
 
 FCameraEvaluationService::~FCameraEvaluationService()
 {
+	if (!ensureMsgf(
+				Evaluator == nullptr, 
+				TEXT("Evaluation service being destroyed without having been torn down properly")))
+	{
+		if (EnumHasAllFlags(PrivateFlags, ECameraEvaluationServiceFlags::NeedsRootCameraNodeEvents))
+		{
+			if (FRootCameraNodeEvaluator* RootNodeEvaluator = Evaluator->GetRootNodeEvaluator())
+			{
+				RootNodeEvaluator->OnCameraRigEvent().RemoveAll(this);
+			}
+		}
+	}
 }
 
 void FCameraEvaluationService::Initialize(const FCameraEvaluationServiceInitializeParams& Params)
 {
+	ensureMsgf(!Evaluator, TEXT("Evaluation service has already been initialized"));
+	ensureMsgf(Params.Evaluator, TEXT("Evaluation service initialized with an invalid system evaluator"));
+	Evaluator = Params.Evaluator;
+
 	OnInitialize(Params);
+
+	if (Evaluator && EnumHasAllFlags(PrivateFlags, ECameraEvaluationServiceFlags::NeedsRootCameraNodeEvents))
+	{
+		FRootCameraNodeEvaluator* RootNodeEvaluator = Params.Evaluator->GetRootNodeEvaluator();
+		if (ensure(RootNodeEvaluator))
+		{
+			RootNodeEvaluator->OnCameraRigEvent().AddRaw(this, &FCameraEvaluationService::RootCameraNodeEventHandler);
+		}
+	}
 }
 
 void FCameraEvaluationService::PreUpdate(const FCameraEvaluationServiceUpdateParams& Params, FCameraEvaluationServiceUpdateResult& OutResult)
@@ -37,17 +65,32 @@ void FCameraEvaluationService::PostUpdate(const FCameraEvaluationServiceUpdatePa
 
 void FCameraEvaluationService::Teardown(const FCameraEvaluationServiceTeardownParams& Params)
 {
+	ensureMsgf(
+			Params.Evaluator && Evaluator == Params.Evaluator,
+			TEXT("Evaluation service torn down with an invalid system evaluator"));
+
 	OnTeardown(Params);
+
+	if (Evaluator && EnumHasAllFlags(PrivateFlags, ECameraEvaluationServiceFlags::NeedsRootCameraNodeEvents))
+	{
+		FRootCameraNodeEvaluator* RootNodeEvaluator = Evaluator->GetRootNodeEvaluator();
+		if (ensure(RootNodeEvaluator))
+		{
+			RootNodeEvaluator->OnCameraRigEvent().RemoveAll(this);
+		}
+	}
+
+	Evaluator = nullptr;
 }
 
-void FCameraEvaluationService::AddReferencedObjects(FReferenceCollector& Collector)
+void FCameraEvaluationService::RootCameraNodeEventHandler(const FRootCameraNodeCameraRigEvent& InEvent)
 {
-	OnAddReferencedObjects(Collector);
+	OnRootCameraNodeEvent(InEvent);
 }
 
-void FCameraEvaluationService::NotifyRootCameraNodeEvent(const FRootCameraNodeCameraRigEvent& InEvent)
+void FCameraEvaluationService::AddReferencedObjects(FReferenceCollector& Collector)
 {
-	OnRootCameraNodeEvent(InEvent);
+	OnAddReferencedObjects(Collector);
 }
 
 bool FCameraEvaluationService::HasAllEvaluationServiceFlags(ECameraEvaluationServiceFlags InFlags) const
diff --git a/Engine/Plugins/Cameras/GameplayCameras/Source/GameplayCameras/Private/Core/CameraSystemEvaluator.cpp b/Engine/Plugins/Cameras/GameplayCameras/Source/GameplayCameras/Private/Core/CameraSystemEvaluator.cpp
--- a/Engine/Plugins/Cameras/GameplayCameras/Source/GameplayCameras/Private/Core/CameraSystemEvaluator.cpp
+++ b/Engine/Plugins/Cameras/GameplayCameras/Source/GameplayCameras/Private/Core/CameraSystemEvaluator.cpp
@@ -212,17 +212,6 @@ TSharedPtr<FCameraEvaluationService> FCameraSystemEvaluator::FindEvaluationServi
 	return nullptr;
 }
 
-void FCameraSystemEvaluator::NotifyRootCameraNodeEvent(const FRootCameraNodeCameraRigEvent& InEvent)
-{
-	for (TSharedPtr<FCameraEvaluationService> EvaluationService : EvaluationServices)
-	{
-		if (EvaluationService->HasAllEvaluationServiceFlags(ECameraEvaluationServiceFlags::NeedsRootCameraNodeEvents))
-		{
-			EvaluationService->NotifyRootCameraNodeEvent(InEvent);
-		}
-	}
-}
-
 void FCameraSystemEvaluator::Update(const FCameraSystemEvaluationParams& Params)
 {
 	UpdateImpl(Params.DeltaTime, ECameraNodeEvaluationType::Standard);
diff --git a/Engine/Plugins/Cameras/GameplayCameras/Source/GameplayCameras/Private/Core/RootCameraNode.cpp b/Engine/Plugins/Cameras/GameplayCameras/Source/GameplayCameras/Private/Core/RootCameraNode.cpp
--- a/Engine/Plugins/Cameras/GameplayCameras/Source/GameplayCameras/Private/Core/RootCameraNode.cpp
+++ b/Engine/Plugins/Cameras/GameplayCameras/Source/GameplayCameras/Private/Core/RootCameraNode.cpp
@@ -117,11 +117,6 @@ void FRootCameraNodeEvaluator::SetPreVisualLayerResult(const FCameraNodeEvaluati
 
 void FRootCameraNodeEvaluator::BroadcastCameraRigEvent(const FRootCameraNodeCameraRigEvent& InEvent) const
 {
-	if (ensure(OwningEvaluator))
-	{
-		OwningEvaluator->NotifyRootCameraNodeEvent(InEvent);
-	}
-
 	OnCameraRigEventDelegate.Broadcast(InEvent);
 }
 
diff --git a/Engine/Plugins/Cameras/GameplayCameras/Source/GameplayCameras/Public/Core/CameraEvaluationService.h b/Engine/Plugins/Cameras/GameplayCameras/Source/GameplayCameras/Public/Core/CameraEvaluationService.h
--- a/Engine/Plugins/Cameras/GameplayCameras/Source/GameplayCameras/Public/Core/CameraEvaluationService.h
+++ b/Engine/Plugins/Cameras/GameplayCameras/Source/GameplayCameras/Public/Core/CameraEvaluationService.h
@@ -101,8 +101,6 @@ public:
 
 	// Internal API.
 
-	void NotifyRootCameraNodeEvent(const FRootCameraNodeCameraRigEvent& InEvent);
-
 	ECameraEvaluationServiceFlags GetEvaluationServiceFlags() const { return PrivateFlags; }
 	bool HasAllEvaluationServiceFlags(ECameraEvaluationServiceFlags InFlags) const;
 
@@ -130,6 +128,10 @@ protected:
 	virtual void OnBuildDebugBlocks(const FCameraDebugBlockBuildParams& Params, FCameraDebugBlockBuilder& Builder) {}
 #endif  // UE_GAMEPLAY_CAMERAS_DEBUG
 
+private:
+
+	void RootCameraNodeEventHandler(const FRootCameraNodeCameraRigEvent& InEvent);
+
 protected:
 
 	/**
@@ -140,6 +142,9 @@ protected:
 
 private:
 
+	/** The camera system that this service lives inside of. */
+	FCameraSystemEvaluator* Evaluator = nullptr;
+
 	/** Evaluation service flags. */
 	ECameraEvaluationServiceFlags PrivateFlags = ECameraEvaluationServiceFlags::None;
 };
diff --git a/Engine/Plugins/Cameras/GameplayCameras/Source/GameplayCameras/Public/Core/CameraSystemEvaluator.h b/Engine/Plugins/Cameras/GameplayCameras/Source/GameplayCameras/Public/Core/CameraSystemEvaluator.h
--- a/Engine/Plugins/Cameras/GameplayCameras/Source/GameplayCameras/Public/Core/CameraSystemEvaluator.h
+++ b/Engine/Plugins/Cameras/GameplayCameras/Source/GameplayCameras/Public/Core/CameraSystemEvaluator.h
@@ -305,8 +305,6 @@ private:
 	void PostCameraDirectorUpdateServices(float DeltaTime, ECameraNodeEvaluationType EvaluationType);
 	void PostUpdateServices(float DeltaTime, ECameraNodeEvaluationType EvaluationType);
 
-	void NotifyRootCameraNodeEvent(const FRootCameraNodeCameraRigEvent& InEvent);
-
 #if UE_GAMEPLAY_CAMERAS_DEBUG
 	static bool IsDebugTraceEnabled();
 	static bool ShouldBuildOrDrawDebugBlocks();
