Cameras: add more module exports

#Synced-CL 49781371

diff --git a/Engine/Plugins/Cameras/GameplayCameras/Source/GameplayCameras/Public/GameFramework/GameplayCamerasPlayerCameraManager.h b/Engine/Plugins/Cameras/GameplayCameras/Source/GameplayCameras/Public/GameFramework/GameplayCamerasPlayerCameraManager.h
--- a/Engine/Plugins/Cameras/GameplayCameras/Source/GameplayCameras/Public/GameFramework/GameplayCamerasPlayerCameraManager.h
+++ b/Engine/Plugins/Cameras/GameplayCameras/Source/GameplayCameras/Public/GameFramework/GameplayCamerasPlayerCameraManager.h
@@ -98,49 +98,49 @@ public:
 
 	/** Activates the given component inside the manager's camera system. */
 	UFUNCTION(BlueprintCallable, Category="Camera")
-	void ActivateGameplayCamera(UGameplayCameraComponentBase* GameplayCamera, EGameplayCameraComponentActivationMode ActivationMode = EGameplayCameraComponentActivationMode::Push);
+	UE_API void ActivateGameplayCamera(UGameplayCameraComponentBase* GameplayCamera, EGameplayCameraComponentActivationMode ActivationMode = EGameplayCameraComponentActivationMode::Push);
 
 	/** Deactivates a previously activated component. */
 	UFUNCTION(BlueprintCallable, Category="Camera")
-	void DeactivateGameplayCamera(UGameplayCameraComponentBase* GameplayCamera, bool bDeactivateAllCameraRigs = false);
+	UE_API void DeactivateGameplayCamera(UGameplayCameraComponentBase* GameplayCamera, bool bDeactivateAllCameraRigs = false);
 
 public:
 
 	/** Activates the given camera rig prefab in the base layer. */
 	UFUNCTION(BlueprintCallable, Category="Camera")
-	FCameraRigInstanceID ActivatePersistentBaseCameraRig(UCameraRigAsset* CameraRig);
+	UE_API FCameraRigInstanceID ActivatePersistentBaseCameraRig(UCameraRigAsset* CameraRig);
 
 	/** Activates the given camera rig prefab in the global layer. */
 	UFUNCTION(BlueprintCallable, Category="Camera")
-	FCameraRigInstanceID ActivatePersistentGlobalCameraRig(UCameraRigAsset* CameraRig);
+	UE_API FCameraRigInstanceID ActivatePersistentGlobalCameraRig(UCameraRigAsset* CameraRig);
 
 	/** Activates the given camera rig prefab in the visual layer. */
 	UFUNCTION(BlueprintCallable, Category="Camera")
-	FCameraRigInstanceID ActivatePersistentVisualCameraRig(UCameraRigAsset* CameraRig);
+	UE_API FCameraRigInstanceID ActivatePersistentVisualCameraRig(UCameraRigAsset* CameraRig);
 
 	/** Deactivates a previously activated base, global, or visual camera rig. */
 	UFUNCTION(BlueprintCallable, Category="Camera")
-	void DeactivateCameraRig(FCameraRigInstanceID InstanceID, bool bImmediately = false);
+	UE_API void DeactivateCameraRig(FCameraRigInstanceID InstanceID, bool bImmediately = false);
 
 public:
 
 	/** Starts a camera modifier rig on the global layer. */
 	UFUNCTION(BlueprintCallable, Category="Camera")
-	FCameraRigInstanceID StartGlobalCameraModifierRig(const UCameraRigAsset* CameraRig, int32 OrderKey = 0);
+	UE_API FCameraRigInstanceID StartGlobalCameraModifierRig(const UCameraRigAsset* CameraRig, int32 OrderKey = 0);
 
 	/** Starts a camera modifier rig on the visual layer. */
 	UFUNCTION(BlueprintCallable, Category="Camera")
-	FCameraRigInstanceID StartVisualCameraModifierRig(const UCameraRigAsset* CameraRig, int32 OrderKey = 0);
+	UE_API FCameraRigInstanceID StartVisualCameraModifierRig(const UCameraRigAsset* CameraRig, int32 OrderKey = 0);
 
 	/** Stops a camera modifier rig on previously started on the global or visual layer. */
 	UFUNCTION(BlueprintCallable, Category="Camera")
-	void StopCameraModifierRig(FCameraRigInstanceID InstanceID, bool bImmediately = false);
+	UE_API void StopCameraModifierRig(FCameraRigInstanceID InstanceID, bool bImmediately = false);
 
 public:
 
 	/** Starts a new camera shake. */
 	UFUNCTION(BlueprintCallable, Category="Camera")
-	FCameraShakeInstanceID StartCameraShakeAsset(
+	UE_API FCameraShakeInstanceID StartCameraShakeAsset(
 			const UCameraShakeAsset* CameraShake,
 			float ShakeScale = 1.f,
 			ECameraShakePlaySpace PlaySpace = ECameraShakePlaySpace::CameraLocal,
@@ -148,11 +148,11 @@ public:
 
 	/** Checks if a camera shake is running. */
 	UFUNCTION(BlueprintCallable, Category="Camera")
-	bool IsCameraShakeAssetPlaying(FCameraShakeInstanceID InInstanceID) const;
+	UE_API bool IsCameraShakeAssetPlaying(FCameraShakeInstanceID InInstanceID) const;
 
 	/** Stops a running camera shake. */
 	UFUNCTION(BlueprintCallable, Category="Camera")
-	bool StopCameraShakeAsset(FCameraShakeInstanceID InInstanceID, bool bImmediately = false);
+	UE_API bool StopCameraShakeAsset(FCameraShakeInstanceID InInstanceID, bool bImmediately = false);
 
 public:
 
