Cameras: don't assert when (re)activating a Gameplay Camera component that is already active

#Synced-CL 49781371

diff --git a/Engine/Plugins/Cameras/GameplayCameras/Source/GameplayCameras/Private/GameFramework/GameplayCameraComponentBase.cpp b/Engine/Plugins/Cameras/GameplayCameras/Source/GameplayCameras/Private/GameFramework/GameplayCameraComponentBase.cpp
--- a/Engine/Plugins/Cameras/GameplayCameras/Source/GameplayCameras/Private/GameFramework/GameplayCameraComponentBase.cpp
+++ b/Engine/Plugins/Cameras/GameplayCameras/Source/GameplayCameras/Private/GameFramework/GameplayCameraComponentBase.cpp
@@ -145,7 +145,8 @@ bool UGameplayCameraComponentBase::EnsureCameraSystemHostIfNeeded()
 	}
 
 	// Create a hosted camera system, and then create and active the evalution context.
-	if (!HasCameraSystem())
+	const bool bHadCameraSystem = HasCameraSystem();
+	if (!bHadCameraSystem)
 	{
 		AActor* OwnerActor = GetOwner();
 		UE_LOG(LogCameraSystem, Log, 
@@ -163,6 +164,7 @@ bool UGameplayCameraComponentBase::EnsureCameraSystemHostIfNeeded()
 		InitializeCameraSystem(Params);
 	}
 
+	// We need our evaluation context to run anything.
 	if (!ensureMsgf(
 				EvaluationContext.IsValid(),
 				TEXT("Can't activate Gameplay Camera component '{0}.{1}': failed to create evaluation context!"),
@@ -171,22 +173,35 @@ bool UGameplayCameraComponentBase::EnsureCameraSystemHostIfNeeded()
 		return false;
 	}
 
+	// If we already had a camera system running, we probably already activated our evaluation context and have 
+	// nothing to do. However, if we just created the camera system, then our evaluation context should not be active
+	// yet, and we'll activate it now.
 	if (!ensureMsgf(
-				!EvaluationContext->IsActive(),
+				bHadCameraSystem || !EvaluationContext->IsActive(),
 				TEXT("Can't activate Gameplay Camera component '{0}.{1}': it is already active!"),
 				*GetNameSafe(GetOwner()), *GetNameSafe(this)))
 	{
 		return false;
 	}
 
-	UE_LOG(LogCameraSystem, Log, 
-			TEXT("Activating gameplay camera '%s.%s' with its hosted camera system."),
-			*GetNameSafe(GetOwner()), *GetNameSafe(this));
-
 	TSharedPtr<FCameraSystemEvaluator> HostedEvaluator = GetCameraSystemEvaluator();
-	check(HostedEvaluator.IsValid());
-	FCameraEvaluationContextStack& ContextStack = HostedEvaluator->GetEvaluationContextStack();
-	ContextStack.PushContext(EvaluationContext.ToSharedRef());
+	if (ensure(HostedEvaluator.IsValid()))
+	{
+		if (!EvaluationContext->IsActive())
+		{
+			UE_LOG(LogCameraSystem, Log, 
+					TEXT("Activating gameplay camera '%s.%s' with its hosted camera system."),
+					*GetNameSafe(GetOwner()), *GetNameSafe(this));
+
+			FCameraEvaluationContextStack& ContextStack = HostedEvaluator->GetEvaluationContextStack();
+			ContextStack.PushContext(EvaluationContext.ToSharedRef());
+		}
+
+		ensureMsgf(
+				EvaluationContext->GetCameraSystemEvaluator() == HostedEvaluator,
+				TEXT("Gameplay Camera Component '{0}.{1}' has a mismatch between evaluation context and camera system!"),
+				*GetNameSafe(GetOwner()), *GetNameSafe(this));
+	}
 
 	if (ensure(OutputCameraComponent))
 	{
