Cameras: fix DPI scaling issue in debug renderer

#Synced-CL 49781371

diff --git a/Engine/Plugins/Cameras/GameplayCameras/Source/GameplayCameras/Private/Debug/CameraDebugRenderer.cpp b/Engine/Plugins/Cameras/GameplayCameras/Source/GameplayCameras/Private/Debug/CameraDebugRenderer.cpp
--- a/Engine/Plugins/Cameras/GameplayCameras/Source/GameplayCameras/Private/Debug/CameraDebugRenderer.cpp
+++ b/Engine/Plugins/Cameras/GameplayCameras/Source/GameplayCameras/Private/Debug/CameraDebugRenderer.cpp
@@ -158,7 +158,7 @@ void FCameraDebugRenderer::Initialize(UWorld* InWorld, const FSceneView* InScene
 			ViewRect = InSceneView->UnconstrainedViewRect;
 		}
 
-		CanvasSize = FVector2d(ViewRect.Width(), ViewRect.Height());
+		CanvasSize = FVector2d(ViewRect.Width(), ViewRect.Height()) / InCanvas->GetDPIScale();
 
 		NextCardPosition = FVector2f{ 
 			(float)CanvasSize.X - (float)GGameplayCamerasDebugCardWidth - (float)GGameplayCamerasDebugRightMargin,
