Cameras: improve fixup of incorrectly saved DefaultParameters property bags

FInstancedPropertyBag used to miss saving its property flags, until it was fixed in 5.8.0. Before that, loading DefaultParameters on camera objects would load default property flags, and when we build the camera object, different flags would be generated and it would dirty the asset. The fix for that was to handle it in 5.7 builds, but we also need to handle it for assets saved in 5.7 and opened in 5.8 builds. This change addresses that.

#Synced-CL 49781371

diff --git a/Engine/Plugins/Cameras/GameplayCameras/Source/GameplayCameras/Private/Build/CameraObjectInterfaceParameterBuilder.cpp b/Engine/Plugins/Cameras/GameplayCameras/Source/GameplayCameras/Private/Build/CameraObjectInterfaceParameterBuilder.cpp
--- a/Engine/Plugins/Cameras/GameplayCameras/Source/GameplayCameras/Private/Build/CameraObjectInterfaceParameterBuilder.cpp
+++ b/Engine/Plugins/Cameras/GameplayCameras/Source/GameplayCameras/Private/Build/CameraObjectInterfaceParameterBuilder.cpp
@@ -8,6 +8,7 @@
 #include "Core/CameraParameters.h"  // IWYU pragma: keep
 #include "Core/CameraVariableReferences.h"  // IWYU pragma: keep
 #include "Core/ICustomCameraNodeParameterProvider.h"
+#include "Misc/EngineVersionComparison.h"
 #include "StructUtils/PropertyBag.h"
 #include "UObject/UnrealType.h"
 
@@ -410,8 +411,6 @@ DoneSearchingRawSourceValuePtr:
 	return bSuccess;
 }
 
-#if UE_VERSION_OLDER_THAN(5,8,0)
-
 void FCameraObjectInterfaceParameterBuilder::FixUpDefaultParameterProperties(TConstArrayView<FCameraObjectInterfaceParameterDefinition> ParameterDefinitions, FInstancedPropertyBag& InOutPropertyBag)
 {
 	const UPropertyBag* PropertyBag = InOutPropertyBag.GetPropertyBagStruct();
@@ -476,7 +475,5 @@ void FCameraObjectInterfaceParameterBuilder::FixUpDefaultParameterProperties(TCo
 	}
 }
 
-#endif
-
 }  // namespace UE::Cameras
 
diff --git a/Engine/Plugins/Cameras/GameplayCameras/Source/GameplayCameras/Private/Core/BaseCameraObject.cpp b/Engine/Plugins/Cameras/GameplayCameras/Source/GameplayCameras/Private/Core/BaseCameraObject.cpp
--- a/Engine/Plugins/Cameras/GameplayCameras/Source/GameplayCameras/Private/Core/BaseCameraObject.cpp
+++ b/Engine/Plugins/Cameras/GameplayCameras/Source/GameplayCameras/Private/Core/BaseCameraObject.cpp
@@ -6,6 +6,7 @@
 #include "Core/CameraNode.h"
 #include "Core/ObjectTreeGraphRootObject.h"
 #include "Misc/EngineVersionComparison.h"
+#include "StructUtils/OverridablePropertyBag.h"
 
 #include UE_INLINE_GENERATED_CPP_BY_NAME(BaseCameraObject)
 
@@ -64,6 +65,23 @@ FCameraObjectConnection* FCameraObjectConnections::FindByTarget(UObject* InTarge
 			});
 }
 
+void UBaseCameraObject::Serialize(FArchive& Ar)
+{
+	// Before 5.8, property bags are never saved with property flags, so always fix them up in PostLoad.
+	// With 5.8, that bug is fixed so we only fix up the property bags if the asset was saved before.
+#if UE_VERSION_NEWER_THAN_OR_EQUAL(5,8,0)
+	Ar.UsingCustomVersion(FOverridablePropertyBagCustomVersion::GUID);
+#endif
+
+	Super::Serialize(Ar);
+
+#if UE_VERSION_NEWER_THAN_OR_EQUAL(5,8,0)
+	bDefaultParametersMayHaveMissingPropertyFlags = Ar.CustomVer(FOverridablePropertyBagCustomVersion::GUID) < FOverridablePropertyBagCustomVersion::MissingPropertyFlags;
+#else
+	bDefaultParametersMayHaveMissingPropertyFlags = true;
+#endif
+}
+
 void UBaseCameraObject::PostLoad()
 {
 	Super::PostLoad();
@@ -73,10 +91,12 @@ void UBaseCameraObject::PostLoad()
 		Guid = FGuid::NewGuid();
 	}
 
-#if UE_VERSION_OLDER_THAN(5,8,0)
-	using namespace UE::Cameras;
-	FCameraObjectInterfaceParameterBuilder::FixUpDefaultParameterProperties(ParameterDefinitions, DefaultParameters);
-#endif
+	if (bDefaultParametersMayHaveMissingPropertyFlags)
+	{
+		using namespace UE::Cameras;
+		FCameraObjectInterfaceParameterBuilder::FixUpDefaultParameterProperties(ParameterDefinitions, DefaultParameters);
+		bDefaultParametersMayHaveMissingPropertyFlags = false;
+	}
 }
 
 void UBaseCameraObject::PostInitProperties()
diff --git a/Engine/Plugins/Cameras/GameplayCameras/Source/GameplayCameras/Private/Core/CameraAsset.cpp b/Engine/Plugins/Cameras/GameplayCameras/Source/GameplayCameras/Private/Core/CameraAsset.cpp
--- a/Engine/Plugins/Cameras/GameplayCameras/Source/GameplayCameras/Private/Core/CameraAsset.cpp
+++ b/Engine/Plugins/Cameras/GameplayCameras/Source/GameplayCameras/Private/Core/CameraAsset.cpp
@@ -4,18 +4,16 @@
 
 #include "Build/CameraAssetBuilder.h"
 #include "Build/CameraBuildContext.h"
+#include "Build/CameraObjectInterfaceParameterBuilder.h"
 #include "Core/CameraDirector.h"
 #include "Core/CameraRigAsset.h"
 #include "Misc/EngineVersionComparison.h"
 #include "Misc/EngineVersionComparison.h"
+#include "StructUtils/OverridablePropertyBag.h"
 #include "UObject/AssetRegistryTagsContext.h"
 #include "UObject/ObjectRedirector.h"
 #include "UObject/ObjectSaveContext.h"
 
-#if UE_VERSION_OLDER_THAN(5,8,0)
-#include "Build/CameraObjectInterfaceParameterBuilder.h"
-#endif
-
 #include UE_INLINE_GENERATED_CPP_BY_NAME(CameraAsset)
 
 void UCameraAssetInterfaceParameter::PostLoad()
@@ -151,14 +149,33 @@ int32 UCameraAsset::RemoveExitTransition(UCameraRigTransition* InTransition)
 	return NumRemoved;
 }
 
+void UCameraAsset::Serialize(FArchive& Ar)
+{
+	// Before 5.8, property bags are never saved with property flags, so always fix them up in PostLoad.
+	// With 5.8, that bug is fixed so we only fix up the property bags if the asset was saved before.
+#if UE_VERSION_NEWER_THAN_OR_EQUAL(5,8,0)
+	Ar.UsingCustomVersion(FOverridablePropertyBagCustomVersion::GUID);
+#endif
+
+	Super::Serialize(Ar);
+
+#if UE_VERSION_NEWER_THAN_OR_EQUAL(5,8,0)
+	bDefaultParametersMayHaveMissingPropertyFlags = Ar.CustomVer(FOverridablePropertyBagCustomVersion::GUID) < FOverridablePropertyBagCustomVersion::MissingPropertyFlags;
+#else
+	bDefaultParametersMayHaveMissingPropertyFlags = true;
+#endif
+}
+
 void UCameraAsset::PostLoad()
 {
 	Super::PostLoad();
 
-#if UE_VERSION_OLDER_THAN(5,8,0)
-	using namespace UE::Cameras;
-	FCameraObjectInterfaceParameterBuilder::FixUpDefaultParameterProperties(ParameterDefinitions, DefaultParameters);
-#endif
+	if (bDefaultParametersMayHaveMissingPropertyFlags)
+	{
+		using namespace UE::Cameras;
+		FCameraObjectInterfaceParameterBuilder::FixUpDefaultParameterProperties(ParameterDefinitions, DefaultParameters);
+		bDefaultParametersMayHaveMissingPropertyFlags = false;
+	}
 
 #if WITH_EDITOR
 	if (CameraDirector)
diff --git a/Engine/Plugins/Cameras/GameplayCameras/Source/GameplayCameras/Private/StructUtils/OverridablePropertyBag.cpp b/Engine/Plugins/Cameras/GameplayCameras/Source/GameplayCameras/Private/StructUtils/OverridablePropertyBag.cpp
--- a/Engine/Plugins/Cameras/GameplayCameras/Source/GameplayCameras/Private/StructUtils/OverridablePropertyBag.cpp
+++ b/Engine/Plugins/Cameras/GameplayCameras/Source/GameplayCameras/Private/StructUtils/OverridablePropertyBag.cpp
@@ -6,20 +6,6 @@
 
 #include UE_INLINE_GENERATED_CPP_BY_NAME(OverridablePropertyBag)
 
-namespace FOverridablePropertyBagCustomVersion
-{
-	enum Type
-	{
-		BeforeCustomVersionWasAdded = 0,
-		FixSerializer = 1,
-
-		VersionPlusOne,
-		LatestVersion = VersionPlusOne - 1
-	};
-
-	const FGuid GUID(0x5426C227, 0x4B3145B2, 0x9B9BED1F, 0x327FB126);
-}
-
 FCustomVersionRegistration GOverridablePropertyBagCustomVersion(
 		FOverridablePropertyBagCustomVersion::GUID,
 		FOverridablePropertyBagCustomVersion::LatestVersion,
@@ -93,7 +79,8 @@ bool FInstancedOverridablePropertyBag::Serialize(FArchive& Ar)
 {
 	Ar.UsingCustomVersion(FOverridablePropertyBagCustomVersion::GUID);
 
-	if (Ar.IsLoading() && Ar.CustomVer(FOverridablePropertyBagCustomVersion::GUID) < FOverridablePropertyBagCustomVersion::FixSerializer)
+	if (Ar.IsLoading() && 
+			Ar.CustomVer(FOverridablePropertyBagCustomVersion::GUID) < FOverridablePropertyBagCustomVersion::FixSerializer)
 	{
 		// There was a short time during which developer data was saved with only default tagged
 		// serialization. This causes data corruption so we avoid that with this use-case here.
diff --git a/Engine/Plugins/Cameras/GameplayCameras/Source/GameplayCameras/Public/Build/CameraObjectInterfaceParameterBuilder.h b/Engine/Plugins/Cameras/GameplayCameras/Source/GameplayCameras/Public/Build/CameraObjectInterfaceParameterBuilder.h
--- a/Engine/Plugins/Cameras/GameplayCameras/Source/GameplayCameras/Public/Build/CameraObjectInterfaceParameterBuilder.h
+++ b/Engine/Plugins/Cameras/GameplayCameras/Source/GameplayCameras/Public/Build/CameraObjectInterfaceParameterBuilder.h
@@ -3,7 +3,6 @@
 #pragma once
 
 #include "Containers/Array.h"
-#include "Misc/EngineVersionComparison.h"
 
 #define UE_API GAMEPLAYCAMERAS_API
 
@@ -36,9 +35,7 @@ public:
 
 	static UE_API bool SetDefaultParameterValue(UBaseCameraObject* CameraObject, const FCameraObjectInterfaceParameterDefinition& ParameterDefinition, UCameraNode* TargetNode, FName TargetPropertyName, bool bAddParameterIfMissing = true);
 
-#if UE_VERSION_OLDER_THAN(5,8,0)
 	static void FixUpDefaultParameterProperties(TConstArrayView<FCameraObjectInterfaceParameterDefinition> ParameterDefinitions, FInstancedPropertyBag& InOutPropertyBag);
-#endif
 
 private:
 
diff --git a/Engine/Plugins/Cameras/GameplayCameras/Source/GameplayCameras/Public/Core/BaseCameraObject.h b/Engine/Plugins/Cameras/GameplayCameras/Source/GameplayCameras/Public/Core/BaseCameraObject.h
--- a/Engine/Plugins/Cameras/GameplayCameras/Source/GameplayCameras/Public/Core/BaseCameraObject.h
+++ b/Engine/Plugins/Cameras/GameplayCameras/Source/GameplayCameras/Public/Core/BaseCameraObject.h
@@ -181,6 +181,7 @@ public:
 public:
 
 	// UObject interface.
+	virtual void Serialize(FArchive& Ar) override;
 	virtual void PostLoad() override;
 	virtual void PostInitProperties() override;
 	virtual void PostDuplicate(EDuplicateMode::Type DuplicateMode) override;
@@ -206,6 +207,9 @@ private:
 	UPROPERTY()
 	TArray<FCameraObjectInterfaceParameterDefinition> ParameterDefinitions;
 
+	// Flag for DefaultParameters maybe having missing property flags.
+	bool bDefaultParametersMayHaveMissingPropertyFlags = false;
+
 	friend class UE::Cameras::FCameraObjectInterfaceParameterBuilder;
 };
 
diff --git a/Engine/Plugins/Cameras/GameplayCameras/Source/GameplayCameras/Public/Core/CameraAsset.h b/Engine/Plugins/Cameras/GameplayCameras/Source/GameplayCameras/Public/Core/CameraAsset.h
--- a/Engine/Plugins/Cameras/GameplayCameras/Source/GameplayCameras/Public/Core/CameraAsset.h
+++ b/Engine/Plugins/Cameras/GameplayCameras/Source/GameplayCameras/Public/Core/CameraAsset.h
@@ -213,6 +213,7 @@ protected:
 #endif
 
 	// UObject interface.
+	virtual void Serialize(FArchive& Ar) override;
 	virtual void PreSave(FObjectPreSaveContext ObjectSaveContext) override;
 	virtual void PostLoad() override;
 	virtual void GetAssetRegistryTags(FAssetRegistryTagsContext Context) const override;
@@ -278,6 +279,8 @@ private:
 
 #endif  // WITH_EDITORONLY_DATA
 
+	// Flag for DefaultParameters maybe having missing property flags.
+	bool bDefaultParametersMayHaveMissingPropertyFlags = false;
 	// Flag for whether we need to upgrade the interface on the next build.
 	bool bNeedsDefaultInterface = false;
 
diff --git a/Engine/Plugins/Cameras/GameplayCameras/Source/GameplayCameras/Public/StructUtils/OverridablePropertyBag.h b/Engine/Plugins/Cameras/GameplayCameras/Source/GameplayCameras/Public/StructUtils/OverridablePropertyBag.h
--- a/Engine/Plugins/Cameras/GameplayCameras/Source/GameplayCameras/Public/StructUtils/OverridablePropertyBag.h
+++ b/Engine/Plugins/Cameras/GameplayCameras/Source/GameplayCameras/Public/StructUtils/OverridablePropertyBag.h
@@ -93,3 +93,18 @@ struct TStructOpsTypeTraits<FInstancedOverridablePropertyBag> : public TStructOp
 	};
 };
 
+namespace FOverridablePropertyBagCustomVersion
+{
+	enum Type
+	{
+		BeforeCustomVersionWasAdded = 0,
+		FixSerializer = 1,
+		MissingPropertyFlags = 2,
+
+		VersionPlusOne,
+		LatestVersion = VersionPlusOne - 1
+	};
+
+	const FGuid GUID(0x5426C227, 0x4B3145B2, 0x9B9BED1F, 0x327FB126);
+}
+
