Cameras: add icons for Camera Rig Actor and Component

#Synced-CL 49781371

diff --git a/Engine/Plugins/Cameras/GameplayCameras/Source/GameplayCamerasEditor/Private/Styles/GameplayCamerasEditorStyle.cpp b/Engine/Plugins/Cameras/GameplayCameras/Source/GameplayCamerasEditor/Private/Styles/GameplayCamerasEditorStyle.cpp
--- a/Engine/Plugins/Cameras/GameplayCameras/Source/GameplayCamerasEditor/Private/Styles/GameplayCamerasEditorStyle.cpp
+++ b/Engine/Plugins/Cameras/GameplayCameras/Source/GameplayCamerasEditor/Private/Styles/GameplayCamerasEditorStyle.cpp
@@ -49,6 +49,10 @@ FGameplayCamerasEditorStyle::FGameplayCamerasEditorStyle()
 	Set("ClassThumbnail.CameraShakeAsset", new IMAGE_BRUSH_SVG("Icons/ContentBrowser-CameraShake", Icon64x64));
 
 	// Camera actors and components.
+	Set("ClassIcon.GameplayCameraRigComponent", new IMAGE_BRUSH_SVG("Icons/GameplayCameraRig_16", Icon16x16));
+	Set("ClassThumbnail.GameplayCameraRigComponent", new IMAGE_BRUSH_SVG("Icons/GameplayCameraRig_64", Icon64x64));
+	Set("ClassIcon.GameplayCameraRigActor", new IMAGE_BRUSH_SVG("Icons/GameplayCameraRig_16", Icon16x16));
+	Set("ClassThumbnail.GameplayCameraRigActor", new IMAGE_BRUSH_SVG("Icons/GameplayCameraRig_64", Icon64x64));
 	Set("ClassIcon.GameplayCameraComponent", new IMAGE_BRUSH_SVG("Icons/GameplayCamera_16", Icon16x16));
 	Set("ClassThumbnail.GameplayCameraComponent", new IMAGE_BRUSH_SVG("Icons/GameplayCamera_64", Icon64x64));
 	Set("ClassIcon.GameplayCameraActor", new IMAGE_BRUSH_SVG("Icons/GameplayCamera_16", Icon16x16));
