Cameras: fix missing pin on nodes in camera shake graphs

#Synced-CL 49781371

diff --git a/Engine/Plugins/Cameras/GameplayCameras/Source/GameplayCamerasEditor/Private/Editors/CameraShakeCameraNodeGraphSchema.cpp b/Engine/Plugins/Cameras/GameplayCameras/Source/GameplayCamerasEditor/Private/Editors/CameraShakeCameraNodeGraphSchema.cpp
--- a/Engine/Plugins/Cameras/GameplayCameras/Source/GameplayCamerasEditor/Private/Editors/CameraShakeCameraNodeGraphSchema.cpp
+++ b/Engine/Plugins/Cameras/GameplayCameras/Source/GameplayCamerasEditor/Private/Editors/CameraShakeCameraNodeGraphSchema.cpp
@@ -4,6 +4,7 @@
 
 #include "Core/CameraShakeAsset.h"
 #include "Core/ShakeCameraNode.h"
+#include "Editors/CameraNodeGraphNode.h"
 #include "Editors/ObjectTreeGraphConfig.h"
 #include "GameplayCamerasEditorSettings.h"
 #include "Nodes/Blends/SimpleBlendCameraNode.h"
@@ -41,6 +42,11 @@ void UCameraShakeCameraNodeGraphSchema::OnBuildGraphConfig(FObjectTreeGraphConfi
 		.HasSelfPin(false)
 		.NodeTitleUsesObjectName(true)
 		.NodeTitleColor(Settings->CameraShakeAssetTitleColor);
+	InOutGraphConfig.ObjectClassConfigs.Emplace(UShakeCameraNode::StaticClass())
+		.StripDisplayNameSuffix(TEXT("Camera Node"))
+		.CreateCategoryMetaData(TEXT("CameraNodeCategories"))
+		.NodeTitleColor(Settings->CameraNodeTitleColor)
+		.GraphNodeClass(UCameraNodeGraphNode::StaticClass());
 }
 
 #undef LOCTEXT_NAMESPACE
