Cameras: fix rewind debugger crash when recording post-process settings with weighted blendables

#Synced-CL 49781371

diff --git a/Engine/Plugins/Cameras/GameplayCameras/Source/GameplayCameras/Private/Debug/CameraSystemTrace.cpp b/Engine/Plugins/Cameras/GameplayCameras/Source/GameplayCameras/Private/Debug/CameraSystemTrace.cpp
--- a/Engine/Plugins/Cameras/GameplayCameras/Source/GameplayCameras/Private/Debug/CameraSystemTrace.cpp
+++ b/Engine/Plugins/Cameras/GameplayCameras/Source/GameplayCameras/Private/Debug/CameraSystemTrace.cpp
@@ -271,6 +271,111 @@ private:
 	TMap<FCameraDebugBlock*, FRelatedIndices> RelatedIndices;
 };
 
+class FCameraDebugBlockWriterArchive : public FBufferArchive
+{
+public:
+
+	// FArchive interface.
+	virtual FArchive& operator<<(UObject*& Res) override { return SaveObjectPtrImpl(Res); }
+	FArchive& operator<<(FLazyObjectPtr& Value) override { return SaveObjectPtrImpl(Value.Get()); }
+	FArchive& operator<<(FObjectPtr& Value) override { return SaveObjectPtrImpl(Value.Get()); }
+	FArchive& operator<<(FSoftObjectPtr& Value) override { return SaveObjectPtrImpl(Value.Get()); }
+	FArchive& operator<<(FSoftObjectPath& Value) override { return SaveObjectPtrImpl(Value.ResolveObject()); }
+	FArchive& operator<<(FWeakObjectPtr& Value) override { return SaveObjectPtrImpl(Value.Get()); }
+
+private:
+
+	FArchive& SaveObjectPtrImpl(UObject* Obj)
+	{
+		if (ensure(IsSaving()))
+		{
+			// See FCameraDebugBlockReaderArchive for information.
+			FString ResClassName = Obj ? Obj->GetClass()->GetName() : FString();
+			*this << ResClassName;
+
+			FObjectKey ResKey(Obj);
+			*this << ResKey;
+		}
+		return *this;
+	}
+};
+
+class FCameraDebugBlockReaderArchive : public FMemoryReader
+{
+public:
+
+	explicit FCameraDebugBlockReaderArchive(const TArray<uint8>& InBytes)
+		: FMemoryReader(InBytes)
+	{}
+
+	// FArchive interface.
+	virtual FArchive& operator<<(UObject*& Res) override { return LoadObjectPtrImpl(Res); }
+	FArchive& operator<<(FLazyObjectPtr& Value) override { return LoadObjectPtrImpl2(Value); }
+	FArchive& operator<<(FObjectPtr& Value) override { return LoadObjectPtrImpl2(Value); }
+	FArchive& operator<<(FSoftObjectPtr& Value) override { return LoadObjectPtrImpl2(Value); }
+	FArchive& operator<<(FSoftObjectPath& Value) override { return LoadObjectPtrImpl2(Value); }
+	FArchive& operator<<(FWeakObjectPtr& Value) override { return LoadObjectPtrImpl2(Value); }
+
+private:
+
+	template<typename PtrType>
+	FArchive& LoadObjectPtrImpl2(PtrType& Value)
+	{
+		UObject* Res = nullptr;
+		LoadObjectPtrImpl(Res);
+		Value = PtrType(Res);
+		return *this;
+	}
+
+	FArchive& LoadObjectPtrImpl(UObject*& Res)
+	{
+		if (ensure(IsLoading()))
+		{
+			// In 99.99% of cases, we are recording gameplay and playing it back immediately after. In this case, the 
+			// object-key that we saved into the trace file will either be:
+			// - A valid object that is still in memory (probably the most common case)
+			// - An object that is gone, and so the object-key resolves to a null pointer.
+			//
+			// This means that we can't (and shouldn't) use UObjects inside debug blocks. And we generally don't. 
+			// At this time, the only exception is the Post Process Settings debug block, which saves the 
+			// FPostProcessSettings structure into the trace file. There is a small chance that it has some 
+			// WeightedBlendables in it, and those are pointers to objects. We won't get debug info for those _if_ 
+			// they are dynamic objects that are gone from memory by the time we replay. That's an acceptable 
+			// limitation for now (especially since, again, in a majority of cases, these post-process materials are 
+			// assets on disk so they're valid during replay).
+			//
+			// We add an extra protection by saving the object's class name. This is for the very unlikely case of
+			// someone saving a trace file and replaying it later, in a different editor session. The object-key would
+			// resolve to some completely unrelated object, so we skip that case. There is of course the even more
+			// unlikely case that the object-key somehow resolves to an object of the right class, in which case we
+			// will happily set it, but since it's the right class it should work, even though it's incorrect. At this
+			// point let's just leave that as a known bug for now.
+			//
+			FString ResClassName;
+			*this << ResClassName;
+
+			FObjectKey ResKey;
+			*this << ResKey;
+
+			if (UObject* LoadedRes = ResKey.ResolveObjectPtr())
+			{
+				if (ResClassName == LoadedRes->GetClass()->GetName())
+				{
+					Res = LoadedRes;
+				}
+			}
+			else
+			{
+				if (ResClassName.IsEmpty())
+				{
+					Res = nullptr;
+				}
+			}
+		}
+		return *this;
+	}
+};
+
 // This name must match the one passed to UE_TRACE_CHANNEL above.
 FString FCameraSystemTrace::ChannelName("CameraSystemChannel");
 // These two names must match the names passed to UE_TRACE_EVENT_BEGIN above.
@@ -289,7 +394,7 @@ void FCameraSystemTrace::TraceEvaluation(UWorld* InWorld, const FCameraSystemEva
 		return;
 	}
 
-	FBufferArchive BufferArchive;
+	FCameraDebugBlockWriterArchive BufferArchive;
 	FCameraDebugBlockWriter Writer(BufferArchive);
 	Writer.Write(InRootDebugBlock);
 
@@ -310,7 +415,7 @@ void FCameraSystemTrace::TraceEvaluation(UWorld* InWorld, const FCameraSystemEva
 
 FCameraDebugBlock* FCameraSystemTrace::ReadEvaluationTrace(TArray<uint8> InSerializedBlocks, FCameraDebugBlockStorage& InStorage)
 {
-	FMemoryReader MemoryArchive(InSerializedBlocks);
+	FCameraDebugBlockReaderArchive MemoryArchive(InSerializedBlocks);
 	FCameraDebugBlockReader Reader(MemoryArchive, InStorage);
 	return Reader.Read();
 }
