Cameras: move the CameraPose-to-MinimalViewInfo code for easy re-use

#Synced-CL 49781371

diff --git a/Engine/Plugins/Cameras/GameplayCameras/Source/GameplayCameras/Private/Core/CameraNodeEvaluator.cpp b/Engine/Plugins/Cameras/GameplayCameras/Source/GameplayCameras/Private/Core/CameraNodeEvaluator.cpp
--- a/Engine/Plugins/Cameras/GameplayCameras/Source/GameplayCameras/Private/Core/CameraNodeEvaluator.cpp
+++ b/Engine/Plugins/Cameras/GameplayCameras/Source/GameplayCameras/Private/Core/CameraNodeEvaluator.cpp
@@ -73,6 +73,17 @@ void FCameraNodeEvaluationResult::LerpAll(const FCameraNodeEvaluationResult& ToR
 	bIsValid = (bIsValid && ToResult.bIsValid);
 }
 
+void FCameraNodeEvaluationResult::GetViewInfo(FMinimalViewInfo& OutViewInfo) const
+{
+	// Get the view info from the camera pose.
+	CameraPose.GetViewInfo(OutViewInfo);
+
+	OutViewInfo.PostProcessSettings = PostProcessSettings.Get();
+	OutViewInfo.PostProcessBlendWeight = 1.f;
+	// Create the physical camera settings if needed. Don't overwrite settings that were set by hand.
+	CameraPose.ApplyPhysicalCameraSettings(OutViewInfo.PostProcessSettings, false);
+}
+
 void FCameraNodeEvaluationResult::Serialize(FArchive& Ar)
 {
 	CameraPose.SerializeWithFlags(Ar);
diff --git a/Engine/Plugins/Cameras/GameplayCameras/Source/GameplayCameras/Private/Core/CameraPose.cpp b/Engine/Plugins/Cameras/GameplayCameras/Source/GameplayCameras/Private/Core/CameraPose.cpp
--- a/Engine/Plugins/Cameras/GameplayCameras/Source/GameplayCameras/Private/Core/CameraPose.cpp
+++ b/Engine/Plugins/Cameras/GameplayCameras/Source/GameplayCameras/Private/Core/CameraPose.cpp
@@ -231,6 +231,33 @@ double FCameraPose::GetVerticalProjectionOffset() const
 	return 2.0 * SensorVerticalOffset / (CroppedSensorHeight * EffectiveOverscan);
 }
 
+void FCameraPose::GetViewInfo(FMinimalViewInfo& OutViewInfo) const
+{
+	OutViewInfo.Location = Location;
+	OutViewInfo.Rotation = Rotation;
+	OutViewInfo.FOV = GetEffectiveFieldOfView();
+	OutViewInfo.DesiredFOV = OutViewInfo.FOV;
+
+	OutViewInfo.AspectRatio = GetSensorAspectRatio();
+	OutViewInfo.bConstrainAspectRatio = ConstrainAspectRatio;
+	OutViewInfo.AspectRatioAxisConstraint = OverrideAspectRatioAxisConstraint ? 
+		TOptional<EAspectRatioAxisConstraint>(AspectRatioAxisConstraint.GetValue()) : 
+		TOptional<EAspectRatioAxisConstraint>();
+
+	OutViewInfo.ProjectionMode = ProjectionMode;
+	if (ProjectionMode == ECameraProjectionMode::Orthographic)
+	{
+		OutViewInfo.OrthoWidth = OrthographicWidth;
+	}
+
+	OutViewInfo.PerspectiveNearClipPlane = NearClippingPlane > 0.0f ? NearClippingPlane : GNearClippingPlane;
+
+	OutViewInfo.OffCenterProjectionOffset.X = GetHorizontalProjectionOffset();
+	OutViewInfo.OffCenterProjectionOffset.Y = GetVerticalProjectionOffset();
+
+	OutViewInfo.ApplyOverscan(GetOverscan());
+}
+
 bool FCameraPose::ApplyPhysicalCameraSettings(FPostProcessSettings& PostProcessSettings, bool bOverwriteSettings) const
 {
 	if (!EnablePhysicalCamera || PhysicalCameraBlendWeight <= 0.f)
diff --git a/Engine/Plugins/Cameras/GameplayCameras/Source/GameplayCameras/Private/Core/CameraSystemEvaluator.cpp b/Engine/Plugins/Cameras/GameplayCameras/Source/GameplayCameras/Private/Core/CameraSystemEvaluator.cpp
--- a/Engine/Plugins/Cameras/GameplayCameras/Source/GameplayCameras/Private/Core/CameraSystemEvaluator.cpp
+++ b/Engine/Plugins/Cameras/GameplayCameras/Source/GameplayCameras/Private/Core/CameraSystemEvaluator.cpp
@@ -488,35 +488,7 @@ void FCameraSystemEvaluator::ViewRotationPreviewUpdate(const FCameraSystemEvalua
 
 void FCameraSystemEvaluator::GetEvaluatedCameraView(FMinimalViewInfo& DesiredView)
 {
-	const FCameraPose& CameraPose = RootNodeResult.CameraPose;
-	DesiredView.Location = CameraPose.GetLocation();
-	DesiredView.Rotation = CameraPose.GetRotation();
-	DesiredView.FOV = CameraPose.GetEffectiveFieldOfView();
-	DesiredView.DesiredFOV = DesiredView.FOV;
-
-	DesiredView.AspectRatio = CameraPose.GetSensorAspectRatio();
-	DesiredView.bConstrainAspectRatio = CameraPose.GetConstrainAspectRatio();
-	DesiredView.AspectRatioAxisConstraint = CameraPose.GetOverrideAspectRatioAxisConstraint() ?
-		CameraPose.GetAspectRatioAxisConstraint() : TOptional<EAspectRatioAxisConstraint>();
-
-	DesiredView.ProjectionMode = CameraPose.GetProjectionMode();
-	if (CameraPose.GetProjectionMode() == ECameraProjectionMode::Orthographic)
-	{
-		DesiredView.OrthoWidth = CameraPose.GetOrthographicWidth();
-	}
-
-	DesiredView.PerspectiveNearClipPlane = CameraPose.GetNearClippingPlane();
-
-	DesiredView.OffCenterProjectionOffset.X = CameraPose.GetHorizontalProjectionOffset();
-	DesiredView.OffCenterProjectionOffset.Y = CameraPose.GetVerticalProjectionOffset();
-
-	const FPostProcessSettingsCollection& PostProcessSettings = RootNodeResult.PostProcessSettings;
-	DesiredView.PostProcessSettings = PostProcessSettings.Get();
-	DesiredView.PostProcessBlendWeight = 1.f;
-	// Create the physical camera settings if needed. Don't overwrite settings that were set by hand.
-	CameraPose.ApplyPhysicalCameraSettings(DesiredView.PostProcessSettings, false);
-
-	DesiredView.ApplyOverscan(CameraPose.GetOverscan());
+	RootNodeResult.GetViewInfo(DesiredView);
 }
 
 void FCameraSystemEvaluator::ExecuteOperation(FCameraOperation& Operation)
diff --git a/Engine/Plugins/Cameras/GameplayCameras/Source/GameplayCameras/Public/Core/CameraNodeEvaluator.h b/Engine/Plugins/Cameras/GameplayCameras/Source/GameplayCameras/Public/Core/CameraNodeEvaluator.h
--- a/Engine/Plugins/Cameras/GameplayCameras/Source/GameplayCameras/Public/Core/CameraNodeEvaluator.h
+++ b/Engine/Plugins/Cameras/GameplayCameras/Source/GameplayCameras/Public/Core/CameraNodeEvaluator.h
@@ -228,11 +228,14 @@ public:
 	/** Interpolate this result towards the other given result. */
 	UE_API void LerpAll(const FCameraNodeEvaluationResult& ToResult, float BlendFactor, bool bIncludePrivateValues = false);
 
+public:
+
+	/** Export this result into the given view info. */
+	UE_API void GetViewInfo(FMinimalViewInfo& OutViewInfo) const;
+
 	/** Serializes this result to the given archive. */
 	UE_API void Serialize(FArchive& Ar);
 
-public:
-
 	/** Collects objects from the context data table. */
 	UE_API void AddReferencedObjects(FReferenceCollector& Collector);
 
diff --git a/Engine/Plugins/Cameras/GameplayCameras/Source/GameplayCameras/Public/Core/CameraPose.h b/Engine/Plugins/Cameras/GameplayCameras/Source/GameplayCameras/Public/Core/CameraPose.h
--- a/Engine/Plugins/Cameras/GameplayCameras/Source/GameplayCameras/Public/Core/CameraPose.h
+++ b/Engine/Plugins/Cameras/GameplayCameras/Source/GameplayCameras/Public/Core/CameraPose.h
@@ -194,6 +194,11 @@ public:
 	/** Gets the default sensor size. */
 	static UE_API void GetDefaultSensorSize(float& OutSensorWidth, float& OutSensorHeight);
 
+	/** 
+	 * Exports the current camera pose to the given view info.
+	 */
+	UE_API void GetViewInfo(FMinimalViewInfo& OutViewInfo) const;
+
 	/**
 	 * Applies the necessary post-process settings given the current values
 	 * on this camera pose.
