Cameras: remove GetTargetAngles and use proper projection/unprojection for all of the framing nodes' logic

- Make the dolly and panning nodes use projection/unprojection for computing how to recompose their shots.
- Add more utility functions for projection/unprojection.
- Make the projection/unprojection functions handle non-default aspect-ratio axis constraints.
- Fix some issues in GetEffectiveFieldsOfView where it didn't correctly account for non-ideal aspect ratios.
- Remove FCameraFramingMath which isn't used anymore.

#Synced-CL 49781371

diff --git a/Engine/Plugins/Cameras/GameplayCameras/Source/GameplayCameras/Private/Core/CameraPose.cpp b/Engine/Plugins/Cameras/GameplayCameras/Source/GameplayCameras/Private/Core/CameraPose.cpp
--- a/Engine/Plugins/Cameras/GameplayCameras/Source/GameplayCameras/Private/Core/CameraPose.cpp
+++ b/Engine/Plugins/Cameras/GameplayCameras/Source/GameplayCameras/Private/Core/CameraPose.cpp
@@ -193,6 +193,21 @@ void FCameraPose::GetDefaultSensorSize(float& OutSensorWidth, float& OutSensorHe
 	OutSensorHeight = 18.67f;
 }
 
+EAspectRatioAxisConstraint FCameraPose::GetEffectiveAspectRatioAxisConstraint() const
+{
+	// Use the same default as UCameraComponent.
+	return GetEffectiveAspectRatioAxisConstraint(EAspectRatioAxisConstraint::AspectRatio_MaintainXFOV);
+}
+
+EAspectRatioAxisConstraint FCameraPose::GetEffectiveAspectRatioAxisConstraint(EAspectRatioAxisConstraint DefaultAspectRatioAxisConstraint) const
+{
+	if (OverrideAspectRatioAxisConstraint)
+	{
+		return AspectRatioAxisConstraint;
+	}
+	return DefaultAspectRatioAxisConstraint;
+}
+
 double FCameraPose::GetHorizontalProjectionOffset() const
 {
 	// Compute projection offset with similar code to UCineCameraComponent...
diff --git a/Engine/Plugins/Cameras/GameplayCameras/Source/GameplayCameras/Private/Math/CameraAimingMath.cpp b/Engine/Plugins/Cameras/GameplayCameras/Source/GameplayCameras/Private/Math/CameraAimingMath.cpp
--- a/Engine/Plugins/Cameras/GameplayCameras/Source/GameplayCameras/Private/Math/CameraAimingMath.cpp
+++ b/Engine/Plugins/Cameras/GameplayCameras/Source/GameplayCameras/Private/Math/CameraAimingMath.cpp
@@ -78,6 +78,57 @@ bool FCameraAimingMath::ComputeTwoBonesCorrection(const FCameraPose& CurrentPose
 	return true;
 }
 
+bool FCameraAimingMath::ComputeTwoBonesCorrection(const FCameraPose& CurrentPose, const FRotator3d& AimOffsetAngles, const FVector3d& PivotLocation, const FVector3d& DesiredTarget, FRotator3d& OutCorrection)
+{
+	// See above for the math. The only difference is that we offset CameraAim by AimOffsetAngles.
+	const FVector3d PivotToDesiredTarget = (DesiredTarget - PivotLocation);
+	const double SphereRadius = PivotToDesiredTarget.Length();
+
+	const FVector3d CameraLocation = CurrentPose.GetLocation();
+	const FRotator3d CameraRotation = CurrentPose.GetRotation();
+	const FVector3d OffsetCameraAim = CameraRotation.RotateVector(AimOffsetAngles.RotateVector(FVector3d::ForwardVector));
+
+	double DistanceToX = 0.0;
+	const bool bGotX = RaySphereIntersectExit(CameraLocation, OffsetCameraAim, PivotLocation, SphereRadius, DistanceToX);
+	if (!bGotX)
+	{
+		return false;
+	}
+
+	const FVector3d X = CameraLocation + OffsetCameraAim * DistanceToX;
+	const FVector3d PX = (X - PivotLocation);
+	const FVector3d PD = (DesiredTarget - PivotLocation);
+	const FRotator3d RotPX = PX.ToOrientationRotator();
+	const FRotator3d RotPD = PD.ToOrientationRotator();
+	OutCorrection = RotPD - RotPX;
+	return true;
+}
+
+bool FCameraAimingMath::ComputeTwoBonesCorrection(const FCameraPose& CurrentPose, const FVector3d& CustomAim, const FVector3d& PivotLocation, const FVector3d& DesiredTarget, FRotator3d& OutCorrection)
+{
+	// See above for the math. The only difference is that we use a custom camera aim vector.
+	const FVector3d PivotToDesiredTarget = (DesiredTarget - PivotLocation);
+	const double SphereRadius = PivotToDesiredTarget.Length();
+
+	const FVector3d CameraLocation = CurrentPose.GetLocation();
+	const FVector3d NormalizedCustomAim = CustomAim.GetUnsafeNormal();
+
+	double DistanceToX = 0.0;
+	const bool bGotX = RaySphereIntersectExit(CameraLocation, NormalizedCustomAim, PivotLocation, SphereRadius, DistanceToX);
+	if (!bGotX)
+	{
+		return false;
+	}
+
+	const FVector3d X = CameraLocation + NormalizedCustomAim * DistanceToX;
+	const FVector3d PX = (X - PivotLocation);
+	const FVector3d PD = (DesiredTarget - PivotLocation);
+	const FRotator3d RotPX = PX.ToOrientationRotator();
+	const FRotator3d RotPD = PD.ToOrientationRotator();
+	OutCorrection = RotPD - RotPX;
+	return true;
+}
+
 const FCameraRigJoint* FCameraAimingMath::FindPivotJoint(const FCameraNodeEvaluationResult& Result)
 {
 	return FindPivotJoint(Result.CameraRigJoints);
@@ -200,18 +251,32 @@ bool FCameraAimingMath::RaySphereIntersectExit(const FVector3d& RayStart, const
 	return false;
 }
 
-double FCameraAimingMath::GetErrorAngle(const FVector3d& AimA, const FVector3d& AimB)
+double FCameraAimingMath::GetAngleBetween(const FVector3d& AimA, const FVector3d& AimB)
+{
+	const FQuat4d Quat = FQuat4d::FindBetweenVectors(AimA, AimB);
+	return Quat.GetAngle();
+}
+
+FRotator3d FCameraAimingMath::GetRotationBetween(const FVector3d& AimA, const FVector3d& AimB)
+{
+	const FQuat4d Quat = FQuat4d::FindBetweenVectors(AimA, AimB);
+	return Quat.Rotator();
+}
+
+FRotator3d FCameraAimingMath::GetNoRollRotationBetween(const FVector3d& AimA, const FVector3d& AimB)
 {
-	const double MultipliedLengths = AimA.Length() * AimB.Length();
-	if (MultipliedLengths > UE_SMALL_NUMBER)
+	const double SqLenA = AimA.SquaredLength();
+	const double SqLenB = AimB.SquaredLength();
+	if (SqLenA > UE_DOUBLE_SMALL_NUMBER && SqLenB > UE_DOUBLE_SMALL_NUMBER)
 	{
-		const double OrthLength = AimA.Cross(AimB).Length();
-		const double ErrorAngle = FMath::Asin(OrthLength / MultipliedLengths);
-		return FMath::RadiansToDegrees(ErrorAngle);
+		const FRotator3d RotationA = AimA.ToOrientationRotator();
+		const FRotator3d RotationB = AimB.ToOrientationRotator();
+		ensure(RotationA.Roll == 0.0 && RotationB.Roll == 0.0);
+		return (RotationB - RotationA);
 	}
 	else
 	{
-		return 0.0;
+		return FRotator3d::ZeroRotator;
 	}
 }
 
diff --git a/Engine/Plugins/Cameras/GameplayCameras/Source/GameplayCameras/Private/Math/CameraFramingMath.cpp b/Engine/Plugins/Cameras/GameplayCameras/Source/GameplayCameras/Private/Math/CameraFramingMath.cpp
deleted file mode 100644
diff --git a/Engine/Plugins/Cameras/GameplayCameras/Source/GameplayCameras/Private/Math/CameraPoseMath.cpp b/Engine/Plugins/Cameras/GameplayCameras/Source/GameplayCameras/Private/Math/CameraPoseMath.cpp
--- a/Engine/Plugins/Cameras/GameplayCameras/Source/GameplayCameras/Private/Math/CameraPoseMath.cpp
+++ b/Engine/Plugins/Cameras/GameplayCameras/Source/GameplayCameras/Private/Math/CameraPoseMath.cpp
@@ -5,45 +5,78 @@
 #include "Core/CameraPose.h"
 #include "Core/CameraEvaluationContext.h"
 #include "CoreGlobals.h"
+#include "Engine/LocalPlayer.h"
+#include "GameFramework/PlayerController.h"
 #include "Math/InverseRotationMatrix.h"
 #include "Math/PerspectiveMatrix.h"
 #include "Math/TranslationMatrix.h"
+#include "SceneView.h"
 
 namespace UE::Cameras
 {
 
+FArchive& operator<< (FArchive& Ar, FCameraFieldsOfView& FieldsOfView)
+{
+	Ar << FieldsOfView.HorizontalFieldOfView;
+	Ar << FieldsOfView.VerticalFieldOfView;
+	return Ar;
+}
+
 FCameraFieldsOfView FCameraPoseMath::GetEffectiveFieldsOfView(const FCameraPose& CameraPose)
 {
-	return GetEffectiveFieldsOfView(CameraPose, CameraPose.GetSensorAspectRatio());
+	return GetEffectiveFieldsOfView(CameraPose, CameraPose.GetSensorAspectRatio(), CameraPose.GetEffectiveAspectRatioAxisConstraint());
 }
 
 FCameraFieldsOfView FCameraPoseMath::GetEffectiveFieldsOfView(const FCameraPose& CameraPose, TSharedPtr<const FCameraEvaluationContext> EvaluationContext)
 {
 	const double AspectRatio = GetEffectiveAspectRatio(CameraPose, EvaluationContext);
-	return GetEffectiveFieldsOfView(CameraPose, AspectRatio);
+	EAspectRatioAxisConstraint DefaultAspectRatioAxisConstraint = GetDefaultAspectRatioAxisConstraint(EvaluationContext);
+	return GetEffectiveFieldsOfView(CameraPose, AspectRatio, DefaultAspectRatioAxisConstraint);
 }
 
-FCameraFieldsOfView FCameraPoseMath::GetEffectiveFieldsOfView(const FCameraPose& CameraPose, double AspectRatio)
+FCameraFieldsOfView FCameraPoseMath::GetEffectiveFieldsOfView(const FCameraPose& CameraPose, double ActualAspectRatio, EAspectRatioAxisConstraint DefaultAspectRatioAxisConstraint)
 {
 	const double HorizontalFOV = CameraPose.GetEffectiveFieldOfView();
+	const double IdealAspectRatio = CameraPose.GetSensorAspectRatio();
 
 	// Check the sort of aspect ratio axis constraint we have.
-	const EAspectRatioAxisConstraint Constraint(CameraPose.GetAspectRatioAxisConstraint());
-	if (Constraint == AspectRatio_MaintainYFOV ||
-			(Constraint == AspectRatio_MajorAxisFOV && AspectRatio < 1.0))
+	const EAspectRatioAxisConstraint Constraint(CameraPose.GetEffectiveAspectRatioAxisConstraint(DefaultAspectRatioAxisConstraint));
+	if (Constraint == AspectRatio_MaintainYFOV || (Constraint == AspectRatio_MajorAxisFOV && ActualAspectRatio < 1.0))
 	{
 		FCameraFieldsOfView FOVs;
 
 		// We need to maintain vertical FOV... the horizontal FOV we have is for our "ideal" aspect ratio,
 		// i.e. our sensor's aspect ratio. Now we need to compute the vertical FOV in this ideal situation
 		// and re-compute the effective horizontal FOV using the effective aspect ratio.
-		const double VerticalFOVRad = 2.0 * FMath::Atan(
-				FMath::Tan(FMath::DegreesToRadians(HorizontalFOV / 2.0)) / AspectRatio);
-		const double HorizontalFOVRad = 2.0 * FMath::Atan(
-				FMath::Tan(VerticalFOVRad / 2.0) * AspectRatio);
-
-		FOVs.HorizontalFieldOfView = FMath::RadiansToDegrees(HorizontalFOVRad);
-		FOVs.VerticalFieldOfView = FMath::RadiansToDegrees(VerticalFOVRad);
+		//
+		// Given H as horizontal FOV, V as vertical FOV, L the focal distance, X the focal plane's width, and Y the
+		// focal plane's height:
+		//
+		//     tan(H/2) = (X/2)/L
+		//     tan(V/2) = (Y/2)/L
+		//
+		// So L = (X/2)/tan(H/2)
+		// And we know that Y = X/A, where A is the aspect ratio (the "ideal" one for now).
+		//
+		//     tan(V/2) = (Y/2)/(X/2)*tan(H/2)
+		//     tan(V/2) = (1/A)*tan(H/2)
+		//     V/2 = atan(tan(H/2)/A)
+		//
+		// Then keep V (since we want to maintain vertical FOV) and recompute H given the "actual" aspect ratio:
+		//
+		//     tan(H/2) = (X/2)/(Y/2)*tan(V/2)
+		//     tan(H/2) = A*tan(V/2)
+		//     H/2 = atan(tan(V/2)*A)
+		//
+		// (where A is the "actual" aspect ratio this time)
+		//
+		const double TanHorizontalHalfFOV = FMath::Tan(FMath::DegreesToRadians(HorizontalFOV / 2.0));
+		const double TanVerticalHalfFOV = TanHorizontalHalfFOV / IdealAspectRatio;
+		const double VerticalHalfFOV = FMath::Atan(TanVerticalHalfFOV);
+		const double HorizontalHalfFOV = FMath::Atan(TanVerticalHalfFOV * ActualAspectRatio);
+
+		FOVs.HorizontalFieldOfView = FMath::RadiansToDegrees(2.0 * HorizontalHalfFOV);
+		FOVs.VerticalFieldOfView = FMath::RadiansToDegrees(2.0 * VerticalHalfFOV);
 
 		return FOVs;
 	}
@@ -54,7 +87,7 @@ FCameraFieldsOfView FCameraPoseMath::GetEffectiveFieldsOfView(const FCameraPose&
 
 		FOVs.HorizontalFieldOfView = HorizontalFOV;
 		FOVs.VerticalFieldOfView = FMath::RadiansToDegrees(2.0 * FMath::Atan(
-					FMath::Tan(FMath::DegreesToRadians(FOVs.HorizontalFieldOfView / 2.0)) / AspectRatio));
+					FMath::Tan(FMath::DegreesToRadians(HorizontalFOV / 2.0)) / ActualAspectRatio));
 
 		return FOVs;
 	}
@@ -80,50 +113,123 @@ double FCameraPoseMath::GetEffectiveAspectRatio(const FCameraPose& CameraPose, T
 	}
 }
 
-FMatrix FCameraPoseMath::BuildProjectionMatrix(const FCameraPose& CameraPose, double AspectRatio)
+EAspectRatioAxisConstraint FCameraPoseMath::GetDefaultAspectRatioAxisConstraint(TSharedPtr<const FCameraEvaluationContext> EvaluationContext)
 {
-	const double NearClippingPlane = CameraPose.GetNearClippingPlane() > 0.0f ? 
-		CameraPose.GetNearClippingPlane() : GNearClippingPlane;
-	const double FieldOfView = FMath::Max(CameraPose.GetEffectiveFieldOfView(), 0.001f);
-
-	FMatrix ProjectionMatrix = FReversedZPerspectiveMatrix(
-			FMath::DegreesToRadians(FieldOfView / 2.0),
-			AspectRatio,
-			1.0f,
-			NearClippingPlane);
-	//ProjectionMatrix = AdjustProjectionMatrixForRHI(ProjectionMatrix);
-	return ProjectionMatrix;
+	if (EvaluationContext)
+	{
+		if (APlayerController* PlayerController = EvaluationContext->GetPlayerController())
+		{
+			if (ULocalPlayer* LocalPlayer = PlayerController->GetLocalPlayer())
+			{
+				return LocalPlayer->AspectRatioAxisConstraint;
+			}
+		}
+	}
+	// Same default as UCameraComponent.
+	return EAspectRatioAxisConstraint::AspectRatio_MaintainXFOV;
 }
 
-FMatrix FCameraPoseMath::BuildViewProjectionMatrix(const FCameraPose& CameraPose, double AspectRatio)
+void FCameraPoseMath::BuildProjectionData(
+		const FCameraPose& CameraPose,
+		TSharedPtr<const FCameraEvaluationContext> EvaluationContext,
+		FSceneViewProjectionData& OutProjectionData)
 {
-	const FTranslationMatrix InverseOrigin(-CameraPose.GetLocation());
+	// Use some arbitrary defaults: 1080p 16/9 screen, and XFOV constraint like in UCameraComponent.
+	static const FIntPoint DefaultViewportSize(1920, 1080);
+	EAspectRatioAxisConstraint DefaultAspectRatioAxisConstraint = AspectRatio_MaintainXFOV;
 
-	// Somehow we need to also transpose... code stolen from ULocalPlayer::GetProjectionData...
-	const FMatrix InverseRotation = FInverseRotationMatrix(CameraPose.GetRotation()) * FMatrix(
-		FPlane(0,	0,	1,	0),
-		FPlane(1,	0,	0,	0),
-		FPlane(0,	1,	0,	0),
-		FPlane(0,	0,	0,	1));
+	FIntPoint ViewportSize(EForceInit::ForceInit);
+	if (EvaluationContext)
+	{
+		ViewportSize = EvaluationContext->GetViewportSize();
 
-	const FMatrix ProjectionMatrix = BuildProjectionMatrix(CameraPose, AspectRatio);
+		if (APlayerController* PlayerController = EvaluationContext->GetPlayerController())
+		{
+			if (ULocalPlayer* LocalPlayer = PlayerController->GetLocalPlayer())
+			{
+				DefaultAspectRatioAxisConstraint = LocalPlayer->AspectRatioAxisConstraint;
+			}
+		}
+	}
+	if (ViewportSize.X == 0 || ViewportSize.Y == 0)
+	{
+		ViewportSize = DefaultViewportSize;
+	}
+
+	BuildProjectionData(CameraPose, ViewportSize, DefaultAspectRatioAxisConstraint, OutProjectionData);
+}
+
+void FCameraPoseMath::BuildProjectionData(
+		const FCameraPose& CameraPose,
+		const FIntPoint& ViewportSize,
+		EAspectRatioAxisConstraint DefaultAspectRatioAxisConstraint,
+		FSceneViewProjectionData& OutProjectionData)
+{
+	FMinimalViewInfo ViewInfo;
+	CameraPose.GetViewInfo(ViewInfo);
+
+	const FIntRect ViewportRect(FIntPoint::ZeroValue, ViewportSize);
+	OutProjectionData.SetViewRectangle(ViewportRect);
+
+	OutProjectionData.ViewOrigin = ViewInfo.Location;
+	OutProjectionData.ViewRotationMatrix = FInverseRotationMatrix(ViewInfo.Rotation) * FMatrix(
+			FPlane(0,	0,	1,	0),
+			FPlane(1,	0,	0,	0),
+			FPlane(0,	1,	0,	0),
+			FPlane(0,	0,	0,	1));
+	FMinimalViewInfo::CalculateProjectionMatrixGivenViewRectangle(
+			ViewInfo, DefaultAspectRatioAxisConstraint, ViewportRect, OutProjectionData);
+	
+	ensure(!OutProjectionData.ProjectionMatrix.ContainsNaN());
+	ensure(!OutProjectionData.ViewRotationMatrix.ContainsNaN());
+}
+
+FMatrix FCameraPoseMath::BuildProjectionMatrix(
+		const FCameraPose& CameraPose,
+		TSharedPtr<const FCameraEvaluationContext> EvaluationContext)
+{
+	FSceneViewProjectionData ProjectionData;
+	BuildProjectionData(CameraPose, EvaluationContext, ProjectionData);
+	return ProjectionData.ProjectionMatrix;
+}
+
+FMatrix FCameraPoseMath::BuildViewProjectionMatrix(
+		const FCameraPose& CameraPose,
+		TSharedPtr<const FCameraEvaluationContext> EvaluationContext)
+{
+	FSceneViewProjectionData ProjectionData;
+	BuildProjectionData(CameraPose, EvaluationContext, ProjectionData);
+	return ProjectionData.ComputeViewProjectionMatrix();
+}
+
+FMatrix FCameraPoseMath::BuildLocalViewProjectionMatrix(
+		const FCameraPose& CameraPose,
+		TSharedPtr<const FCameraEvaluationContext> EvaluationContext)
+{
+	FCameraPose IdentityPose(CameraPose);
+	IdentityPose.SetLocation(FVector3d::ZeroVector);
+	IdentityPose.SetRotation(FRotator3d::ZeroRotator);
 
-	return InverseOrigin * InverseRotation * ProjectionMatrix;
+	FSceneViewProjectionData ProjectionData;
+	BuildProjectionData(IdentityPose, EvaluationContext, ProjectionData);
+	return ProjectionData.ComputeViewProjectionMatrix();
 }
 
 TOptional<FVector2d> FCameraPoseMath::ProjectWorldToScreen(
-		const FCameraPose& CameraPose, double AspectRatio, 
+		const FCameraPose& CameraPose,
+		TSharedPtr<const FCameraEvaluationContext> EvaluationContext,
 		const FVector3d& WorldLocation, bool bForceLocationInsideFrustum)
 {
-	const FMatrix ViewProjectionMatrix = BuildViewProjectionMatrix(CameraPose, AspectRatio);
+	const FMatrix ViewProjectionMatrix = BuildViewProjectionMatrix(CameraPose, EvaluationContext);
 	return ProjectToScreen(ViewProjectionMatrix, WorldLocation, bForceLocationInsideFrustum);
 }
 
 TOptional<FVector2d> FCameraPoseMath::ProjectCameraToScreen(
-			const FCameraPose& CameraPose, double AspectRatio, 
-			const FVector3d& CameraSpaceLocation, bool bForceLocationInsideFrustum)
+		const FCameraPose& CameraPose,
+		TSharedPtr<const FCameraEvaluationContext> EvaluationContext,
+		const FVector3d& CameraSpaceLocation, bool bForceLocationInsideFrustum)
 {
-	const FMatrix ProjectionMatrix = BuildProjectionMatrix(CameraPose, AspectRatio);
+	const FMatrix ProjectionMatrix = BuildProjectionMatrix(CameraPose, EvaluationContext);
 	return ProjectToScreen(ProjectionMatrix, CameraSpaceLocation, bForceLocationInsideFrustum);
 }
 
@@ -160,30 +266,44 @@ TOptional<FVector2d> FCameraPoseMath::ProjectToScreen(
 	return FVector2d(ScreenSpaceX, ScreenSpaceY);
 }
 
-FRay3d FCameraPoseMath::UnprojectScreenToCamera(const FCameraPose& CameraPose, double AspectRatio, const FVector2D& ScreenSpacePoint)
+FRay3d FCameraPoseMath::UnprojectScreenToCamera(
+		const FCameraPose& CameraPose,
+		TSharedPtr<const FCameraEvaluationContext> EvaluationContext,
+		const FVector2D& ScreenSpacePoint)
 {
-	const FMatrix ProjectionMatrix = BuildProjectionMatrix(CameraPose, AspectRatio);
+	const FMatrix ProjectionMatrix = BuildProjectionMatrix(CameraPose, EvaluationContext);
 	const FMatrix InvProjectionMatrix = ProjectionMatrix.InverseFast();
 	return UnprojectFromScreen(InvProjectionMatrix, ScreenSpacePoint);
 }
 
-FVector3d FCameraPoseMath::UnprojectScreenToCamera(const FCameraPose& CameraPose, double AspectRatio, const FVector2D& ScreenSpacePoint, double PredictedDistance)
+FVector3d FCameraPoseMath::UnprojectScreenToCamera(
+		const FCameraPose& CameraPose,
+		TSharedPtr<const FCameraEvaluationContext> EvaluationContext,
+		const FVector2D& ScreenSpacePoint,
+		double PredictedDistance)
 {
-	const FRay3d UnprojectedRay = UnprojectScreenToCamera(CameraPose, AspectRatio, ScreenSpacePoint);
+	const FRay3d UnprojectedRay = UnprojectScreenToCamera(CameraPose, EvaluationContext, ScreenSpacePoint);
 	const FVector3d WorldPoint = UnprojectedRay.PointAt(PredictedDistance);
 	return WorldPoint;
 }
 
-FRay3d FCameraPoseMath::UnprojectScreenToWorld(const FCameraPose& CameraPose, double AspectRatio, const FVector2D& ScreenSpacePoint)
+FRay3d FCameraPoseMath::UnprojectScreenToWorld(
+		const FCameraPose& CameraPose,
+		TSharedPtr<const FCameraEvaluationContext> EvaluationContext,
+		const FVector2D& ScreenSpacePoint)
 {
-	const FMatrix ViewProjectionMatrix = BuildViewProjectionMatrix(CameraPose, AspectRatio);
+	const FMatrix ViewProjectionMatrix = BuildViewProjectionMatrix(CameraPose, EvaluationContext);
 	const FMatrix InvViewProjectionMatrix = ViewProjectionMatrix.InverseFast();
 	return UnprojectFromScreen(InvViewProjectionMatrix, ScreenSpacePoint);
 }
 
-FVector3d FCameraPoseMath::UnprojectScreenToWorld(const FCameraPose& CameraPose, double AspectRatio, const FVector2D& ScreenSpacePoint, double PredictedDistance)
+FVector3d FCameraPoseMath::UnprojectScreenToWorld(
+		const FCameraPose& CameraPose,
+		TSharedPtr<const FCameraEvaluationContext> EvaluationContext,
+		const FVector2D& ScreenSpacePoint,
+		double PredictedDistance)
 {
-	const FRay3d UnprojectedRay = UnprojectScreenToWorld(CameraPose, AspectRatio, ScreenSpacePoint);
+	const FRay3d UnprojectedRay = UnprojectScreenToWorld(CameraPose, EvaluationContext, ScreenSpacePoint);
 	const FVector3d WorldPoint = UnprojectedRay.PointAt(PredictedDistance);
 	return WorldPoint;
 }
diff --git a/Engine/Plugins/Cameras/GameplayCameras/Source/GameplayCameras/Private/Nodes/Framing/BaseFramingCameraNode.cpp b/Engine/Plugins/Cameras/GameplayCameras/Source/GameplayCameras/Private/Nodes/Framing/BaseFramingCameraNode.cpp
--- a/Engine/Plugins/Cameras/GameplayCameras/Source/GameplayCameras/Private/Nodes/Framing/BaseFramingCameraNode.cpp
+++ b/Engine/Plugins/Cameras/GameplayCameras/Source/GameplayCameras/Private/Nodes/Framing/BaseFramingCameraNode.cpp
@@ -181,9 +181,6 @@ void FBaseFramingCameraNodeEvaluator::ComputeCurrentState(const FCameraNodeEvalu
 	// Get screen-space coordinates of the ideal framing point. These are in 0..1 UI space.
 	State.IdealTarget = Readers.IdealFramingLocation.Get(OutResult.VariableTable);
 
-	// Cache the aspect ratio.
-	State.AspectRatio = FCameraPoseMath::GetEffectiveAspectRatio(OutResult.CameraPose, Params.EvaluationContext);
-
 	// Update the damping factors and reengage/disengage times in case they are driven by a variable.
 	State.ReframeDampingFactor = Readers.ReframeDampingFactor.Get(OutResult.VariableTable);
 	State.LowReframeDampingFactor = Readers.LowReframeDampingFactor.Get(OutResult.VariableTable);
@@ -208,17 +205,18 @@ void FBaseFramingCameraNodeEvaluator::ComputeCurrentState(const FCameraNodeEvalu
 	// shot transform.
 	FCameraPose TempPose(OutResult.CameraPose);
 	TempPose.SetTransform(LastFraming);
+	const FMatrix TempPoseViewProjectionMatrix = FCameraPoseMath::BuildViewProjectionMatrix(TempPose, Params.EvaluationContext);
 
 	// Process our targets and figure out the weighted average we should be aiming at.
 	FVector2d NewScreenTarget;
 	FVector3d OldWorldTarget = State.WorldTarget;
-	ComputeFinalTargetInfo(TempPose, State.WorldTarget, NewScreenTarget, State.ScreenTargetBounds);
+	ComputeFinalTargetInfo(TempPose, TempPoseViewProjectionMatrix, State.WorldTarget, NewScreenTarget, State.ScreenTargetBounds);
 
 	// See if we need to extrapolate where the target will be in "anticipation time" seconds.
 	State.ScreenTarget = ComputeAnticipatedScreenTarget(Params.DeltaTime, State.ScreenTarget, NewScreenTarget);
 
 	// Compute how fast the target has moved over the last frame.
-	State.TargetScreenSpaceSpeed = ComputeScreenTargetSpeed(Params.DeltaTime, TempPose, OldWorldTarget, State.WorldTarget);
+	State.TargetScreenSpaceSpeed = ComputeScreenTargetSpeed(Params.DeltaTime, TempPoseViewProjectionMatrix, OldWorldTarget, State.WorldTarget);
 
 	// Compute the effective dead-zone, which is the subset of the dead-zone that encompasses as much
 	// of the target's bound as possible.
@@ -361,7 +359,7 @@ void FBaseFramingCameraNodeEvaluator::ComputeCurrentState(const FCameraNodeEvalu
 	}
 }
 
-bool FBaseFramingCameraNodeEvaluator::ComputeFinalTargetInfo(const FCameraPose& CameraPose, FVector3d& OutWorldTarget, FVector2d& OutScreenTarget, FFramingZone& OutScreenBounds)
+bool FBaseFramingCameraNodeEvaluator::ComputeFinalTargetInfo(const FCameraPose& CameraPose, const FMatrix& ViewProjectionMatrix, FVector3d& OutWorldTarget, FVector2d& OutScreenTarget, FFramingZone& OutScreenBounds)
 {
 	TConstArrayView<FCameraActorComputedTargetInfo> TargetInfos(WorldTargets.TargetInfos);
 
@@ -392,11 +390,11 @@ bool FBaseFramingCameraNodeEvaluator::ComputeFinalTargetInfo(const FCameraPose&
 		FComputedTargetScreenInfo& TargetScreenInfo(TargetScreenInfos[Index]);
 
 		const FVector3d& WorldTarget = TargetInfo.Transform.GetLocation();
-		const TOptional<FVector2d> ScreenTarget = FCameraPoseMath::ProjectWorldToScreen(CameraPose, State.AspectRatio, WorldTarget, true);
+		const TOptional<FVector2d> ScreenTarget = FCameraPoseMath::ProjectToScreen(ViewProjectionMatrix, WorldTarget, true);
 
 		TargetScreenInfo.WorldTarget = WorldTarget;
 		TargetScreenInfo.ScreenTarget = ScreenTarget.Get(FVector2d(0.5, 0.5));
-		TargetScreenInfo.ScreenBounds = ComputeScreenTargetBounds(CameraPose, State.AspectRatio, TargetInfo.Transform, TargetInfo.LocalBounds);
+		TargetScreenInfo.ScreenBounds = ComputeScreenTargetBounds(ViewProjectionMatrix, TargetInfo.Transform, TargetInfo.LocalBounds);
 		TargetScreenInfo.WorldTargetDistance = FVector3d::Distance(CameraPose.GetLocation(), WorldTarget);
 		TargetScreenInfo.NormalizedWeight = TargetInfo.NormalizedWeight;
 
@@ -429,7 +427,8 @@ bool FBaseFramingCameraNodeEvaluator::ComputeFinalTargetInfo(const FCameraPose&
 
 	// Unproject the final screen target, and use the weigted average distance to get, roughly, what
 	// world-space target we might be looking at.
-	const FVector3d FinalWorldTarget = FCameraPoseMath::UnprojectScreenToWorld(CameraPose, State.AspectRatio, FinalScreenTarget, FinalWorldTargetDistance);
+	const FMatrix InvViewProjectionMatrix = ViewProjectionMatrix.InverseFast();
+	const FVector3d FinalWorldTarget = FCameraPoseMath::UnprojectFromScreen(InvViewProjectionMatrix, FinalScreenTarget, FinalWorldTargetDistance);
 
 	OutScreenTarget = FinalScreenTarget;
 	OutScreenBounds = FinalScreenBounds;
@@ -438,7 +437,7 @@ bool FBaseFramingCameraNodeEvaluator::ComputeFinalTargetInfo(const FCameraPose&
 	return true;
 }
 
-float FBaseFramingCameraNodeEvaluator::ComputeScreenTargetSpeed(float DeltaTime, const FCameraPose& CameraPose, const FVector3d& OldWorldTarget, const FVector3d& NewWorldTarget)
+float FBaseFramingCameraNodeEvaluator::ComputeScreenTargetSpeed(float DeltaTime, const FMatrix& ViewProjectionMatrix, const FVector3d& OldWorldTarget, const FVector3d& NewWorldTarget)
 {
 	if (DeltaTime <= 0.f)
 	{
@@ -446,8 +445,8 @@ float FBaseFramingCameraNodeEvaluator::ComputeScreenTargetSpeed(float DeltaTime,
 		return State.TargetScreenSpaceSpeed;
 	}
 
-	TOptional<FVector2d> OldScreenTarget = FCameraPoseMath::ProjectWorldToScreen(CameraPose, State.AspectRatio, OldWorldTarget, true);
-	TOptional<FVector2d> NewScreenTarget = FCameraPoseMath::ProjectWorldToScreen(CameraPose, State.AspectRatio, NewWorldTarget, true);
+	TOptional<FVector2d> OldScreenTarget = FCameraPoseMath::ProjectToScreen(ViewProjectionMatrix, OldWorldTarget, true);
+	TOptional<FVector2d> NewScreenTarget = FCameraPoseMath::ProjectToScreen(ViewProjectionMatrix, NewWorldTarget, true);
 	if (!OldScreenTarget.IsSet() || !NewScreenTarget.IsSet())
 	{
 		return State.TargetScreenSpaceSpeed;
@@ -520,7 +519,7 @@ FFramingZone FBaseFramingCameraNodeEvaluator::ComputeEffectiveDeadZone()
 	return EffectiveDeadZone;
 }
 
-FFramingZone FBaseFramingCameraNodeEvaluator::ComputeScreenTargetBounds(const FCameraPose& CameraPose, float AspectRatio, const FTransform3d& TargetTransform, const FBoxSphereBounds3d& LocalBounds)
+FFramingZone FBaseFramingCameraNodeEvaluator::ComputeScreenTargetBounds(const FMatrix& ViewProjectionMatrix, const FTransform3d& TargetTransform, const FBoxSphereBounds3d& LocalBounds)
 {
 	const FVector3d BoxExtent = LocalBounds.BoxExtent;
 	FVector3d BoxCorners[8];
@@ -536,7 +535,7 @@ FFramingZone FBaseFramingCameraNodeEvaluator::ComputeScreenTargetBounds(const FC
 	FVector2d ScreenBoxCorners[8];
 	for (int32 Index = 0; Index < 8; ++Index)
 	{
-		TOptional<FVector2d> ScreenCorner = FCameraPoseMath::ProjectWorldToScreen(CameraPose, AspectRatio, BoxCorners[Index], true);
+		TOptional<FVector2d> ScreenCorner = FCameraPoseMath::ProjectToScreen(ViewProjectionMatrix, BoxCorners[Index], true);
 		ScreenBoxCorners[Index] = ScreenCorner.GetValue();
 	}
 
@@ -779,6 +778,7 @@ void FBaseFramingCameraNodeEvaluator::FState::Serialize(FArchive& Ar)
 	Ar << DisengageTime;
 	Ar << ToggleEngageTimeLeft;
 	Ar << ToggleEngageAlpha;
+	Ar << TargetMovementAnticipationTime;
 	Ar << DeadZone;
 	Ar << SoftZone;
 
@@ -788,7 +788,9 @@ void FBaseFramingCameraNodeEvaluator::FState::Serialize(FArchive& Ar)
 	Ar << EffectiveDeadZone;
 
 	Ar << TargetFramingState;
+	Ar << TargetScreenSpaceSpeed;
 	Ar << bIsReframingTarget;
+
 	Ar << ReframeDamper;
 }
 
diff --git a/Engine/Plugins/Cameras/GameplayCameras/Source/GameplayCameras/Private/Nodes/Framing/DollyFramingCameraNode.cpp b/Engine/Plugins/Cameras/GameplayCameras/Source/GameplayCameras/Private/Nodes/Framing/DollyFramingCameraNode.cpp
--- a/Engine/Plugins/Cameras/GameplayCameras/Source/GameplayCameras/Private/Nodes/Framing/DollyFramingCameraNode.cpp
+++ b/Engine/Plugins/Cameras/GameplayCameras/Source/GameplayCameras/Private/Nodes/Framing/DollyFramingCameraNode.cpp
@@ -4,16 +4,12 @@
 
 #include "Core/CameraEvaluationContext.h"
 #include "Core/CameraParameterReader.h"
-#include "Core/CameraVariableAssets.h"
 #include "Debug/CameraDebugBlock.h"
 #include "Debug/CameraDebugBlockBuilder.h"
 #include "Debug/CameraDebugRenderer.h"
-#include "GameFramework/Pawn.h"
-#include "GameFramework/PlayerController.h"
 #include "GameplayCameras.h"
-#include "Math/CameraFramingMath.h"
+#include "Math/CameraAimingMath.h"
 #include "Math/CameraPoseMath.h"
-#include "Math/InverseRotationMatrix.h"
 
 #include UE_INLINE_GENERATED_CPP_BY_NAME(DollyFramingCameraNode)
 
@@ -142,40 +138,26 @@ FTransform3d FDollyFramingCameraNodeEvaluator::BuildDollyShotTransform(const FCa
 
 FVector3d FDollyFramingCameraNodeEvaluator::ComputeFramingTranslation(const FCameraPose& CameraPose, TSharedPtr<const FCameraEvaluationContext> EvaluationContext)
 {
-	const float AspectRatio = FCameraPoseMath::GetEffectiveAspectRatio(CameraPose, EvaluationContext);
-	const FCameraFieldsOfView FOVs(FCameraPoseMath::GetEffectiveFieldsOfView(CameraPose, AspectRatio));
-
 	// Get the position of the current target in camera space.
 	const FTransform3d InverseCameraTransform = CameraPose.GetTransform().Inverse();
 	const FVector3d TargetInCameraSpace = InverseCameraTransform.TransformPosition(State.WorldTarget);
 
-	// Get the horizontal and vertical angles, relative to the aiming vector, for the desired 
-	// target position.
-	const FVector2d DesiredAngles = FCameraFramingMath::GetTargetAngles(Desired.ScreenTarget, FOVs);
+	// Get the direction, relative to the aiming vector, for the desired target position.
+	const FMatrix LocalViewProjectionMatrix = FCameraPoseMath::BuildLocalViewProjectionMatrix(CameraPose, EvaluationContext);
+	const FMatrix InvLocalViewProjectionMatrix = LocalViewProjectionMatrix.InverseFast();
+	const FRay3d DesiredAim = FCameraPoseMath::UnprojectFromScreen(InvLocalViewProjectionMatrix, Desired.ScreenTarget);
 
-	// The position of the desired target in camera space can be deduced from the angle and the current
-	// target position. We know that:
-	//
-	// 1) The current and desired targets will be on a plane parallel to the focal plane (i.e. on a plane
-	//    orthogonal to the aiming vector). This is because our dolly shot only translate laterally and
-	//    vertically -- it doesn't translate forwards/backwards.
-	//
-	// 2) We can do basic trigonometry for each axis (left/right and up/down, a.k.a. Y and Z in camera 
-	//    space). There's a right triangle between the aim vector (up to the targets' distance) and the
-	//    vector from the camera position to the desired target position.
-	//
-	// From (1) we know that the desired target's distance from the focal plane is the same as the
-	// current target's. So the X coordinates (near/far) are the same.
+	// We want the distance to the desired target to be so that it sits on the same Y/Z plane as the current target.
+	// That is, we want to only move laterally and vertically. This means that:
 	//
-	// From (2) we can use the sin() of the horizontal/vertical angles to get the horizontal/vertical
-	// coordinates of the desired target (again, in camera space). We just need to invert the vertical
-	// one because positive Z is up, while positive vertical angle is down (because this was computed
-	// in -1..1 UI screen-space).
+	//     DesiredAim.Direction . LocalUnitX = TargetInCameraSpace.X
+	// 
+	// Since DesiredAim's direction vector is a unit vector, we can therefore do:
 	//
-	const FVector3d DesiredInCameraSpace(
-			TargetInCameraSpace.X,
-			TargetInCameraSpace.X * FMath::Sin(FMath::DegreesToRadians(DesiredAngles.X)),
-			TargetInCameraSpace.X * -FMath::Sin(FMath::DegreesToRadians(DesiredAngles.Y)));
+	const double CosAngle = FVector3d::DotProduct(DesiredAim.Direction, FVector3d::UnitX());
+	ensure(CosAngle != 0.0);
+	const double DesiredAimLength = (CosAngle != 0) ? TargetInCameraSpace.X / CosAngle : TargetInCameraSpace.X;
+	const FVector3d DesiredInCameraSpace = DesiredAim.Direction * DesiredAimLength;
 
 	// Now we can figure out the desired camera-space offset that the dolly needs to move by. Remember
 	// that, for instance, moving the camera to the *right* will result in the target moving to the
diff --git a/Engine/Plugins/Cameras/GameplayCameras/Source/GameplayCameras/Private/Nodes/Framing/PanningFramingCameraNode.cpp b/Engine/Plugins/Cameras/GameplayCameras/Source/GameplayCameras/Private/Nodes/Framing/PanningFramingCameraNode.cpp
--- a/Engine/Plugins/Cameras/GameplayCameras/Source/GameplayCameras/Private/Nodes/Framing/PanningFramingCameraNode.cpp
+++ b/Engine/Plugins/Cameras/GameplayCameras/Source/GameplayCameras/Private/Nodes/Framing/PanningFramingCameraNode.cpp
@@ -9,7 +9,7 @@
 #include "Debug/CameraDebugBlockBuilder.h"
 #include "Debug/CameraDebugRenderer.h"
 #include "GameplayCameras.h"
-#include "Math/CameraFramingMath.h"
+#include "Math/CameraAimingMath.h"
 #include "Math/CameraPoseMath.h"
 #include "Math/InverseRotationMatrix.h"
 
@@ -98,21 +98,17 @@ void FPanningFramingCameraNodeEvaluator::OnRun(const FCameraNodeEvaluationParams
 	{
 		FCameraPose LastShotPose(OutResult.CameraPose);
 		LastShotPose.SetTransform(LastShotTransform);
+		const FMatrix LocalViewProjectionMatrix = FCameraPoseMath::BuildLocalViewProjectionMatrix(LastShotPose, Params.EvaluationContext);
+		const FMatrix InvLocalViewProjectionMatrix = LocalViewProjectionMatrix.InverseFast();
 
-		const float AspectRatio = FCameraPoseMath::GetEffectiveAspectRatio(LastShotPose, Params.EvaluationContext);
-		const FCameraFieldsOfView FOVs(FCameraPoseMath::GetEffectiveFieldsOfView(LastShotPose, AspectRatio));
+		// Get the direction for where the target is currently at, and where we want it to be at.
+		const FRay3d TargetAim = FCameraPoseMath::UnprojectFromScreen(InvLocalViewProjectionMatrix, State.ScreenTarget);
+		const FRay3d DesiredAim = FCameraPoseMath::UnprojectFromScreen(InvLocalViewProjectionMatrix, Desired.ScreenTarget);
 
-		// Get yaw/pitch angles for where the target is currently at, and where we want it to be at.
-		const FVector2d TargetAngles = FCameraFramingMath::GetTargetAngles(State.ScreenTarget, FOVs);
-		const FVector2d DesiredAngles = FCameraFramingMath::GetTargetAngles(Desired.ScreenTarget, FOVs);
-
-		// We need to reverse pitch because rotators treat positive pitch as up, whereas we were 
-		// treating that as negative (due to the -1..1 normalized space of our math).
-		const FRotator3d PanningCorrection(
-				(DesiredAngles.Y - TargetAngles.Y),	// Pitch
-				(TargetAngles.X - DesiredAngles.X),	// Yaw
-				0);									// Roll
-		if (ensure(!PanningCorrection.ContainsNaN()))
+		// Compute the yaw/pitch correction needed. Remember that if we want the target to move left on the screen, we
+		// need to turn to the right. This is why we compute the Desired->Current angle, and not the other way around.
+		const FRotator3d PanningCorrection = FCameraAimingMath::GetNoRollRotationBetween(DesiredAim.Direction, TargetAim.Direction);
+		if (ensure(!PanningCorrection.ContainsNaN() && PanningCorrection.Roll == 0.0))
 		{
 			PanningRotation += PanningCorrection;
 		}
diff --git a/Engine/Plugins/Cameras/GameplayCameras/Source/GameplayCameras/Public/Core/CameraPose.h b/Engine/Plugins/Cameras/GameplayCameras/Source/GameplayCameras/Public/Core/CameraPose.h
--- a/Engine/Plugins/Cameras/GameplayCameras/Source/GameplayCameras/Public/Core/CameraPose.h
+++ b/Engine/Plugins/Cameras/GameplayCameras/Source/GameplayCameras/Public/Core/CameraPose.h
@@ -165,6 +165,12 @@ public:
 	/** Gets the aspect ratio of the camera sensor. */
 	UE_API double GetSensorAspectRatio() const;
 
+	/** Gets the aspect ratio axis constraint if overriden, otherwise a default one. */
+	UE_API EAspectRatioAxisConstraint GetEffectiveAspectRatioAxisConstraint() const;
+
+	/** Gets the aspect ratio axis constraint if overriden, otherwise the provided default one. */
+	UE_API EAspectRatioAxisConstraint GetEffectiveAspectRatioAxisConstraint(EAspectRatioAxisConstraint DefaultAspectRatioAxisConstraint) const;
+
 	/** Gets the horizontal projection offset. */
 	UE_API double GetHorizontalProjectionOffset() const;
 
diff --git a/Engine/Plugins/Cameras/GameplayCameras/Source/GameplayCameras/Public/Math/CameraAimingMath.h b/Engine/Plugins/Cameras/GameplayCameras/Source/GameplayCameras/Public/Math/CameraAimingMath.h
--- a/Engine/Plugins/Cameras/GameplayCameras/Source/GameplayCameras/Public/Math/CameraAimingMath.h
+++ b/Engine/Plugins/Cameras/GameplayCameras/Source/GameplayCameras/Public/Math/CameraAimingMath.h
@@ -42,8 +42,21 @@ public:
 	/** Computes the rotation correction needed for the given camera to be aimed at the given target.*/
 	static bool ComputeTwoBonesCorrection(const FCameraPose& CurrentPose, const FVector3d& PivotLocation, const FVector3d& DesiredTarget, FRotator3d& OutCorrection);
 
+	/**
+	 * As per the other ComputeTwoBonesCorrection function, but take into account an additional angular offset for the
+	 * camera's line of sight. These angles should be expressed in the pivot's coordinate system.
+	 */
+	static bool ComputeTwoBonesCorrection(const FCameraPose& CurrentPose, const FRotator3d& AimOffsetAngles, const FVector3d& PivotLocation, const FVector3d& DesiredTarget, FRotator3d& OutCorrection);
+	static bool ComputeTwoBonesCorrection(const FCameraPose& CurrentPose, const FVector3d& CustomAim, const FVector3d& PivotLocation, const FVector3d& DesiredTarget, FRotator3d& OutCorrection);
+
 	/** Computes the difference between two aim vectors. */
-	static double GetErrorAngle(const FVector3d& AimA, const FVector3d& AimB);
+	static double GetAngleBetween(const FVector3d& AimA, const FVector3d& AimB);
+
+	/** Computes the rotation between two aim vectors. */
+	static FRotator3d GetRotationBetween(const FVector3d& AimA, const FVector3d& AimB);
+
+	/** Computes the rotation between two aim vectors, using only yaw and pitch. */
+	static FRotator3d GetNoRollRotationBetween(const FVector3d& AimA, const FVector3d& AimB);
 
 private:
 
diff --git a/Engine/Plugins/Cameras/GameplayCameras/Source/GameplayCameras/Public/Math/CameraFramingMath.h b/Engine/Plugins/Cameras/GameplayCameras/Source/GameplayCameras/Public/Math/CameraFramingMath.h
deleted file mode 100644
diff --git a/Engine/Plugins/Cameras/GameplayCameras/Source/GameplayCameras/Public/Math/CameraPoseMath.h b/Engine/Plugins/Cameras/GameplayCameras/Source/GameplayCameras/Public/Math/CameraPoseMath.h
--- a/Engine/Plugins/Cameras/GameplayCameras/Source/GameplayCameras/Public/Math/CameraPoseMath.h
+++ b/Engine/Plugins/Cameras/GameplayCameras/Source/GameplayCameras/Public/Math/CameraPoseMath.h
@@ -8,8 +8,11 @@
 
 #define UE_API GAMEPLAYCAMERAS_API
 
+class FArchive;
 struct FCameraFieldsOfView;
 struct FCameraPose;
+struct FSceneViewProjectionData;
+enum EAspectRatioAxisConstraint : int;
 
 namespace UE::Cameras
 {
@@ -24,6 +27,7 @@ struct FCameraFieldsOfView
 	double HorizontalFieldOfView;
 	double VerticalFieldOfView;
 };
+FArchive& operator<< (FArchive& Ar, FCameraFieldsOfView& FieldsOfView);
 
 /**
  * A utility class for mathematical functions related to a camera pose.
@@ -39,30 +43,61 @@ public:
 	UE_API static FCameraFieldsOfView GetEffectiveFieldsOfView(const FCameraPose& CameraPose, TSharedPtr<const FCameraEvaluationContext> EvaluationContext);
 
 	/** Gets both horizontal and vertical effective fields of view, using the given aspect ratio. */
-	UE_API static FCameraFieldsOfView GetEffectiveFieldsOfView(const FCameraPose& CameraPose, double AspectRatio);
+	UE_API static FCameraFieldsOfView GetEffectiveFieldsOfView(const FCameraPose& CameraPose, double ActualAspectRatio, EAspectRatioAxisConstraint DefaultAspectRatioAxisConstraint);
 
 	/** Gets the aspect ratio of the viewport associated with the given player controller. */
 	UE_API static double GetEffectiveAspectRatio(const FCameraPose& CameraPose, TSharedPtr<const FCameraEvaluationContext> EvaluationContext);
 
+	/** Gets the default aspect ratio axis constraint associated with the given player controller. */
+	UE_API static EAspectRatioAxisConstraint GetDefaultAspectRatioAxisConstraint(TSharedPtr<const FCameraEvaluationContext> EvaluationContext);
+
+	/**
+	 * Builds the projection matrices of the given camera pose.
+	 */
+	UE_API static void BuildProjectionData(
+			const FCameraPose& CameraPose,
+			TSharedPtr<const FCameraEvaluationContext> EvaluationContext,
+			FSceneViewProjectionData& OutProjectionData);
+	/**
+	 * Builds the projection matrices of the given camera pose.
+	 */
+	UE_API static void BuildProjectionData(
+			const FCameraPose& CameraPose,
+			const FIntPoint& ViewportSize,
+			EAspectRatioAxisConstraint DefaultAspectRatioAxisConstraint,
+			FSceneViewProjectionData& OutProjectionData);
 	/** 
 	 * Builds the projection matrix of the given camera pose. 
 	 * This matrix is suitable for projecting camera-space points onto screen-space.
 	 */
-	UE_API static FMatrix BuildProjectionMatrix(const FCameraPose& CameraPose, double AspectRatio);
+	UE_API static FMatrix BuildProjectionMatrix(
+			const FCameraPose& CameraPose,
+			TSharedPtr<const FCameraEvaluationContext> EvaluationContext);
 	/** 
 	 * Builds the view-projection matrix of the given camera pose.
 	 * This matrix combines the camera transform and the projection matrix, making it
 	 * suitable for projecting world-space points onto screen-space.
 	 */
-	UE_API static FMatrix BuildViewProjectionMatrix(const FCameraPose& CameraPose, double AspectRatio);
+	UE_API static FMatrix BuildViewProjectionMatrix(
+			const FCameraPose& CameraPose,
+			TSharedPtr<const FCameraEvaluationContext> EvaluationContext);
+	/**
+	 * Builds a view-projection matrix using an identity view, so that using for unprojecting 
+	 * yields camera-local positions.
+	 */
+	UE_API static FMatrix BuildLocalViewProjectionMatrix(
+			const FCameraPose& CameraPose,
+			TSharedPtr<const FCameraEvaluationContext> EvaluationContext);
 
 	/** Projects the given world point onto screen-space. */
 	UE_API static TOptional<FVector2d> ProjectWorldToScreen(
-			const FCameraPose& CameraPose, double AspectRatio, 
+			const FCameraPose& CameraPose,
+			TSharedPtr<const FCameraEvaluationContext> EvaluationContext,
 			const FVector3d& WorldLocation, bool bForceLocationInsideFrustum = false);
 	/** Projects the given camera-local point onto screen-space. */
 	UE_API static TOptional<FVector2d> ProjectCameraToScreen(
-			const FCameraPose& CameraPose, double AspectRatio, 
+			const FCameraPose& CameraPose,
+			TSharedPtr<const FCameraEvaluationContext> EvaluationContext,
 			const FVector3d& CameraSpaceLocation, bool bForceLocationInsideFrustum = false);
 
 	/**
@@ -75,16 +110,27 @@ public:
 
 	/** Unproject a screen-space point into a camera-local ray. */
 	UE_API static FRay3d UnprojectScreenToCamera(
-			const FCameraPose& CameraPose, double AspectRatio, const FVector2D& ScreenSpacePoint);
+			const FCameraPose& CameraPose,
+			TSharedPtr<const FCameraEvaluationContext> EvaluationContext,
+			const FVector2D& ScreenSpacePoint);
 	/** Unproject a screen-space point into a camera-local point given an expected distance. */
 	UE_API static FVector3d UnprojectScreenToCamera(
-			const FCameraPose& CameraPose, double AspectRatio, const FVector2D& ScreenSpacePoint, double PredictedDistance);
+			const FCameraPose& CameraPose, 
+			TSharedPtr<const FCameraEvaluationContext> EvaluationContext,
+			const FVector2D& ScreenSpacePoint,
+			double PredictedDistance);
 	/** Unproject a screen-space point into a world-space ray. */
 	UE_API static FRay3d UnprojectScreenToWorld(
-			const FCameraPose& CameraPose, double AspectRatio, const FVector2D& ScreenSpacePoint);
+			const FCameraPose& CameraPose,
+			TSharedPtr<const FCameraEvaluationContext> EvaluationContext,
+			const FVector2D& ScreenSpacePoint);
 	/** Unproject a screen-space point into a world-space point given an expected distance. */
 	UE_API static FVector3d UnprojectScreenToWorld(
-			const FCameraPose& CameraPose, double AspectRatio, const FVector2D& ScreenSpacePoint, double PredictedDistance);
+			const FCameraPose& CameraPose,
+			TSharedPtr<const FCameraEvaluationContext> EvaluationContext,
+			const FVector2D& ScreenSpacePoint,
+			double PredictedDistance);
+
 	/** Unprojects the given screen-space point into a ray. */
 	UE_API static FRay3d UnprojectFromScreen(
 			const FMatrix& InverseViewProjectionMatrix, const FVector2D& ScreenSpacePoint);
diff --git a/Engine/Plugins/Cameras/GameplayCameras/Source/GameplayCameras/Public/Nodes/Framing/BaseFramingCameraNode.h b/Engine/Plugins/Cameras/GameplayCameras/Source/GameplayCameras/Public/Nodes/Framing/BaseFramingCameraNode.h
--- a/Engine/Plugins/Cameras/GameplayCameras/Source/GameplayCameras/Public/Nodes/Framing/BaseFramingCameraNode.h
+++ b/Engine/Plugins/Cameras/GameplayCameras/Source/GameplayCameras/Public/Nodes/Framing/BaseFramingCameraNode.h
@@ -194,12 +194,12 @@ private:
 	void ComputeDesiredState(const FCameraNodeEvaluationParams& Params, const FCameraNodeEvaluationResult& OutResult);
 
 	bool AcquireTargetInfo(TSharedPtr<const FCameraEvaluationContext> EvaluationContext, const FCameraNodeEvaluationResult& InResult, TArray<FCameraActorComputedTargetInfo>& OutInfos);
-	bool ComputeFinalTargetInfo(const FCameraPose& CameraPose, FVector3d& OutWorldTarget, FVector2d& OutScreenTarget, FFramingZone& OutScreenBounds);
-	float ComputeScreenTargetSpeed(float DeltaTime, const FCameraPose& CameraPose, const FVector3d& OldWorldTarget, const FVector3d& NewWorldTarget);
+	bool ComputeFinalTargetInfo(const FCameraPose& CameraPose, const FMatrix& ViewProjectionMatrix, FVector3d& OutWorldTarget, FVector2d& OutScreenTarget, FFramingZone& OutScreenBounds);
+	float ComputeScreenTargetSpeed(float DeltaTime, const FMatrix& ViewProjectionMatrix, const FVector3d& OldWorldTarget, const FVector3d& NewWorldTarget);
 	FVector2d ComputeAnticipatedScreenTarget(float DeltaTime, const FVector2d& InPreviousAnticipatedScreenTarget, const FVector2d& InScreenTarget);
 	FFramingZone ComputeEffectiveDeadZone();
 
-	static FFramingZone ComputeScreenTargetBounds(const FCameraPose& CameraPose, float AspectRatio, const FTransform3d& TargetTransform, const FBoxSphereBounds3d& LocalBounds);
+	static FFramingZone ComputeScreenTargetBounds(const FMatrix& ViewProjectionMatrix, const FTransform3d& TargetTransform, const FBoxSphereBounds3d& LocalBounds);
 
 protected:
 
@@ -248,8 +248,6 @@ protected:
 	{
 		/** Screen-space position of the ideal framing position. */
 		FVector2d IdealTarget;
-		/** Current aspect ratio of the screen. */
-		float AspectRatio;
 		/** Current reframing damping factor. */
 		float ReframeDampingFactor;
 		/** Current low reframing damping factor. */
