Cameras: add screen-space framing option on the Aim-At camera action

- Adds a Target Framing property for where to frame the target on screen.
- Use the new 2-bone IK function that takes into account an offset from the centre camera aim.
- Improve the debug rendering a bit with extra information about the aiming.
- Change the default lock-on tolerance to be a bit looser.

#Synced-CL 49781371

diff --git a/Engine/Plugins/Cameras/GameplayCameras/Source/GameplayCameras/Private/Actions/AimAtCameraAction.cpp b/Engine/Plugins/Cameras/GameplayCameras/Source/GameplayCameras/Private/Actions/AimAtCameraAction.cpp
--- a/Engine/Plugins/Cameras/GameplayCameras/Source/GameplayCameras/Private/Actions/AimAtCameraAction.cpp
+++ b/Engine/Plugins/Cameras/GameplayCameras/Source/GameplayCameras/Private/Actions/AimAtCameraAction.cpp
@@ -11,6 +11,8 @@
 #include "Debug/CameraDebugBlockBuilder.h"
 #include "Debug/CameraDebugRenderer.h"
 #include "Math/CameraAimingMath.h"
+#include "Math/CameraPoseMath.h"
+#include "Math/ColorList.h"
 #include "Serialization/MemoryReader.h"
 #include "Serialization/MemoryWriter.h"
 #include "Services/CameraActionEvaluator.h"
@@ -38,7 +40,8 @@ protected:
 private:
 
 	void LockUserInputThisFrame(const FCameraActionEvaluationParams& Params);
-	bool RunPreviewEvaluation(const FCameraActionEvaluationParams& Params, const FCameraNodeEvaluationResult& Result, FRotator3d& OutCorrection);
+	void RunPreviewEvaluation(const FCameraActionEvaluationParams& Params, const FCameraNodeEvaluationResult& Result);
+	bool ComputeDesiredCorrection(const FCameraActionEvaluationParams& Params, FRotator3d& OutCorrection);
 	bool ExecuteYawPitchCorrection(const FCameraActionEvaluationParams& Params, const FRotator3d& Correction);
 
 	bool ContinueAimAction(const FCameraActionEvaluationParams& Params, FCameraActionEvaluationResult& OutResult);
@@ -54,6 +57,8 @@ private:
 	TArray<uint8> EvaluatorSnapshot;
 
 	FVector3d TargetLocation;
+	FVector2d TargetFraming = { 0.5, 0.5 };
+	FVector3d TargetFramingAim;
 
 	FCameraPose LastCameraPose;
 	FVector3d LastContextLocation;
@@ -65,6 +70,8 @@ private:
 	float DebugElapsedTime = 0.f;
 	FRotator3d DebugCorrectionLeft;
 	FRotator3d DebugCurrentCorrection;
+	FVector3d DebugCameraLocation;
+	FVector3d DebugCameraAim;
 #endif  // UE_GAMEPLAY_CAMERAS_DEBUG
 };
 
@@ -87,6 +94,7 @@ void FAimAtCameraActionEvaluator::OnInitialize(const FCameraActionEvaluatorIniti
 	ChildHierarchy.Build(CameraRigEvaluationInfo.RootEvaluator);
 
 	TargetLocation = ActionData->TargetLocation;
+	TargetFraming = ActionData->TargetFraming;
 
 #if UE_GAMEPLAY_CAMERAS_DEBUG
 	DebugElapsedTime = 0.f;
@@ -131,7 +139,7 @@ void FAimAtCameraActionEvaluator::LockUserInputThisFrame(const FCameraActionEval
 	ChildHierarchy.CallExecuteOperation(OperationParams, Operation);
 }
 
-bool FAimAtCameraActionEvaluator::RunPreviewEvaluation(const FCameraActionEvaluationParams& Params, const FCameraNodeEvaluationResult& Result, FRotator3d& OutCorrection)
+void FAimAtCameraActionEvaluator::RunPreviewEvaluation(const FCameraActionEvaluationParams& Params, const FCameraNodeEvaluationResult& Result)
 {
 	ScratchResult.OverrideAll(Result, true);
 
@@ -160,7 +168,10 @@ bool FAimAtCameraActionEvaluator::RunPreviewEvaluation(const FCameraActionEvalua
 		FMemoryReader Reader(EvaluatorSnapshot);
 		ChildHierarchy.CallSerialize(SerializeParams, Reader);
 	}
+}
 
+bool FAimAtCameraActionEvaluator::ComputeDesiredCorrection(const FCameraActionEvaluationParams& Params, FRotator3d& OutCorrection)
+{
 	// Find the pivot for this frame.
 	const FCameraRigJoint* PreviewPivot = FCameraAimingMath::FindPivotJoint(ScratchResult.CameraRigJoints);
 	if (!PreviewPivot)
@@ -172,10 +183,33 @@ bool FAimAtCameraActionEvaluator::RunPreviewEvaluation(const FCameraActionEvalua
 		return false;
 	}
 
+	static const FVector2d CenterFraming(0.5, 0.5);
+	const bool bNeedsTargetFraming = (TargetFraming != CenterFraming);
+	if (bNeedsTargetFraming)
+	{
+		TSharedPtr<const FCameraEvaluationContext> EvaluationContext = Params.Scope->GetEvaluationContext();
+		const FRay3d TargetRay = FCameraPoseMath::UnprojectScreenToWorld(
+				ScratchResult.CameraPose, EvaluationContext, TargetFraming);
+		TargetFramingAim = TargetRay.Direction;
+	}
+	else
+	{
+		TargetFramingAim = ScratchResult.CameraPose.GetAimDir();
+	}
+
+#if UE_GAMEPLAY_CAMERAS_DEBUG
+	DebugCameraLocation = ScratchResult.CameraPose.GetLocation();
+	DebugCameraAim = ScratchResult.CameraPose.GetAimDir();
+#endif
+
 	// Compute the total correction need to aim at the target.
 	FRotator3d TotalCorrection;
 	const bool bGotCorrection = FCameraAimingMath::ComputeTwoBonesCorrection(
-			ScratchResult.CameraPose, PreviewPivot->Transform.GetLocation(), TargetLocation, TotalCorrection);
+			ScratchResult.CameraPose, 
+			TargetFramingAim, 
+			PreviewPivot->Transform.GetLocation(), 
+			TargetLocation,
+			TotalCorrection);
 	if (!bGotCorrection)
 	{
 		UE_LOG(LogCameraSystem, Warning, 
@@ -229,18 +263,23 @@ bool FAimAtCameraActionEvaluator::ContinueAimAction(const FCameraActionEvaluatio
 	// Prevent the user from turning the camera this frame, since we want to continue aiming.
 	LockUserInputThisFrame(Params);
 
-	// Run a preview evaluation of the camera rig to see exactly how much we still need to turn to get to the target.
+	// Run a preview evaluation of the camera rig.
+	RunPreviewEvaluation(Params, OutResult.Result);
+
+	// Compute exactly how much we still need to turn to get to the target.
 	FRotator3d TotalCorrection;
-	if (!RunPreviewEvaluation(Params, OutResult.Result, TotalCorrection))
+	if (!ComputeDesiredCorrection(Params, TotalCorrection))
 	{
 		return false;
 	}
+
 #if UE_GAMEPLAY_CAMERAS_DEBUG
 	DebugCorrectionLeft = TotalCorrection;
 #endif
 
 	// See if we're close enough to the target given our tolerance margin.
-	if (FMath::Abs(TotalCorrection.Yaw) < ActionData->LockOnAngleTolerance && FMath::Abs(TotalCorrection.Pitch) < ActionData->LockOnAngleTolerance)
+	if (FMath::Abs(TotalCorrection.Yaw) < ActionData->LockOnAngleTolerance && 
+			FMath::Abs(TotalCorrection.Pitch) < ActionData->LockOnAngleTolerance)
 	{
 		bIsLockedOn = true;
 		return (ActionData->LockOnPolicy == EAimAtCameraActionLockOnPolicy::KeepLock);
@@ -286,16 +325,18 @@ bool FAimAtCameraActionEvaluator::KeepLockOn(const FCameraActionEvaluationParams
 	// Prevent the user from turning the camera this frame, since we want to keep a lock on the target.
 	LockUserInputThisFrame(Params);
 
-	// Run a preview of what the camera will look at this frame, and figure out the correction to keep it locked
-	// on the target.
-	FRotator3d Correction;
-	if (!RunPreviewEvaluation(Params, OutResult.Result, Correction))
+	// Run a preview of what the camera will look at this frame.
+	RunPreviewEvaluation(Params, OutResult.Result);
+
+	// Figure out the correction to keep it locked on the target.
+	FRotator3d TotalCorrection;
+	if (!ComputeDesiredCorrection(Params, TotalCorrection))
 	{
 		return false;
 	}
 
 	// Run the correction on the camera rig.
-	if (!ExecuteYawPitchCorrection(Params, Correction))
+	if (!ExecuteYawPitchCorrection(Params, TotalCorrection))
 	{
 		return false;
 	}
@@ -326,6 +367,10 @@ UE_DECLARE_CAMERA_DEBUG_BLOCK_START(, FAimAtCameraActionDebugBlock)
 	UE_DECLARE_CAMERA_DEBUG_BLOCK_FIELD(float, ElapsedTime)
 	UE_DECLARE_CAMERA_DEBUG_BLOCK_FIELD(FRotator3d, CorrectionLeft)
 	UE_DECLARE_CAMERA_DEBUG_BLOCK_FIELD(FRotator3d, CurrentCorrection)
+	UE_DECLARE_CAMERA_DEBUG_BLOCK_FIELD(FVector2d, TargetFraming)
+	UE_DECLARE_CAMERA_DEBUG_BLOCK_FIELD(FVector3d, TargetFramingAim)
+	UE_DECLARE_CAMERA_DEBUG_BLOCK_FIELD(FVector3d, CameraLocation)
+	UE_DECLARE_CAMERA_DEBUG_BLOCK_FIELD(FVector3d, CameraAim)
 	UE_DECLARE_CAMERA_DEBUG_BLOCK_FIELD(bool, bIsLockedOn)
 UE_DECLARE_CAMERA_DEBUG_BLOCK_END()
 
@@ -335,6 +380,10 @@ void FAimAtCameraActionEvaluator::OnBuildDebugBlocks(const FCameraDebugBlockBuil
 	DebugBlock.ElapsedTime = DebugElapsedTime;
 	DebugBlock.CorrectionLeft = DebugCorrectionLeft;
 	DebugBlock.CurrentCorrection = DebugCurrentCorrection;
+	DebugBlock.TargetFraming = TargetFraming;
+	DebugBlock.TargetFramingAim = TargetFramingAim;
+	DebugBlock.CameraLocation = DebugCameraLocation;
+	DebugBlock.CameraAim = DebugCameraAim;
 	DebugBlock.bIsLockedOn = bIsLockedOn;
 }
 
@@ -349,11 +398,24 @@ void FAimAtCameraActionDebugBlock::OnDebugDraw(const FCameraDebugBlockDrawParams
 	else
 	{
 		Renderer.AddText(
-				TEXT("elapsed time %.1f ; correction yaw=%.1f/%.1f  pitch=%.1f/%.1f"),
+				TEXT("elapsed time %.1f ; correction yaw=%.3f/%.3f  pitch=%.3f/%.3f"),
 				ElapsedTime,
 				CurrentCorrection.Yaw, CorrectionLeft.Yaw,
 				CurrentCorrection.Pitch, CorrectionLeft.Pitch);
 	}
+
+	if (!Renderer.IsExternalRendering())
+	{
+		const FVector2D CanvasSize = Renderer.GetCanvasSize();
+
+		const FVector2D TargetFramingPosition = TargetFraming * CanvasSize;
+		Renderer.Draw2DPointCross(TargetFramingPosition, 20, FColorList::LightGrey);
+	}
+	else
+	{
+		Renderer.DrawLine(CameraLocation, CameraLocation + 2000.f * CameraAim, FColorList::LightBlue);
+		Renderer.DrawLine(CameraLocation, CameraLocation + 2000.f * TargetFramingAim, FColorList::LightBlue);
+	}
 }
 
 #endif  // UE_GAMEPLAY_CAMERAS_DEBUG
diff --git a/Engine/Plugins/Cameras/GameplayCameras/Source/GameplayCameras/Public/Actions/AimAtCameraAction.h b/Engine/Plugins/Cameras/GameplayCameras/Source/GameplayCameras/Public/Actions/AimAtCameraAction.h
--- a/Engine/Plugins/Cameras/GameplayCameras/Source/GameplayCameras/Public/Actions/AimAtCameraAction.h
+++ b/Engine/Plugins/Cameras/GameplayCameras/Source/GameplayCameras/Public/Actions/AimAtCameraAction.h
@@ -37,12 +37,16 @@ public:
 
 	/** The tolerance within which we can consider the aiming to be locked on target. */
 	UPROPERTY(EditAnywhere, BlueprintReadWrite, Category=Common, meta=(ExposeOnSpawn=true))
-	float LockOnAngleTolerance = 0.01f;
+	float LockOnAngleTolerance = 0.05f;
 
 	/** The lock-on policy. */
 	UPROPERTY(EditAnywhere, BlueprintReadWrite, Category=Common, meta=(ExposeOnSpawn=true))
 	EAimAtCameraActionLockOnPolicy LockOnPolicy = EAimAtCameraActionLockOnPolicy::Disengage;
 
+	/** Where to aim the target in screen-space. */
+	UPROPERTY(EditAnywhere, BlueprintReadWrite, Category=Aiming, meta=(ExposeOnSpawn=true))
+	FVector2D TargetFraming = { 0.5, 0.5 };
+
 protected:
 
 	// UCameraAction interface.
