Cameras: improve debug info for player controller/view target setups

Add information about viewport size and aspect ratio.

#Synced-CL 49781371

diff --git a/Engine/Plugins/Cameras/GameplayCameras/Source/GameplayCameras/Private/Debug/PlayerControllersDebugBlock.cpp b/Engine/Plugins/Cameras/GameplayCameras/Source/GameplayCameras/Private/Debug/PlayerControllersDebugBlock.cpp
--- a/Engine/Plugins/Cameras/GameplayCameras/Source/GameplayCameras/Private/Debug/PlayerControllersDebugBlock.cpp
+++ b/Engine/Plugins/Cameras/GameplayCameras/Source/GameplayCameras/Private/Debug/PlayerControllersDebugBlock.cpp
@@ -54,18 +54,31 @@ void FPlayerControllersDebugBlock::Initialize(UWorld* World)
 		FPlayerControllerDebugInfo PlayerControllerInfo;
 		PlayerControllerInfo.PlayerControllerName = *GetNameSafe(PlayerController);
 		PlayerControllerInfo.CameraManagerName = *GetNameSafe(CameraManager);
-		PlayerControllerInfo.ViewTargetName = *GetNameSafe(ActiveViewTarget);
 
+		int32 ViewportSizeX = 0, ViewportSizeY = 0;
+		PlayerController->GetViewportSize(ViewportSizeX, ViewportSizeY);
+		PlayerControllerInfo.ViewportSize = FIntPoint(ViewportSizeX, ViewportSizeY);
+
+		PlayerControllerInfo.bHasDefaultAspectRatioAxisConstraint = false;
 		if (const ULocalPlayer* LocalPlayer = PlayerController->GetLocalPlayer())
 		{
 			PlayerControllerInfo.LocalPlayerName = LocalPlayer->GetName();
 			PlayerControllerInfo.DefaultAspectRatioAxisConstraint = LocalPlayer->AspectRatioAxisConstraint;
+			PlayerControllerInfo.bHasDefaultAspectRatioAxisConstraint = true;
 		}
 
+		PlayerControllerInfo.ViewTargetName = *GetNameSafe(ActiveViewTarget);
 		PlayerControllerInfo.ViewTargetLocation = ViewTargetPOV.Location;
 		PlayerControllerInfo.ViewTargetRotation = ViewTargetPOV.Rotation;
 		PlayerControllerInfo.ViewTargetFOV = ViewTargetPOV.FOV;
 		PlayerControllerInfo.ViewTargetAspectRatio = ViewTargetPOV.AspectRatio;
+		PlayerControllerInfo.bViewTargetConstrainAspectRatio = ViewTargetPOV.bConstrainAspectRatio;
+		PlayerControllerInfo.bViewTargetOverrideAspectRatioAxisConstraint = false;
+		if (ViewTargetPOV.AspectRatioAxisConstraint.IsSet())
+		{
+			PlayerControllerInfo.ViewTargetAspectRatioAxisConstraint = ViewTargetPOV.AspectRatioAxisConstraint.GetValue();
+			PlayerControllerInfo.bViewTargetOverrideAspectRatioAxisConstraint = true;
+		}
 
 		PlayerControllers.Add(PlayerControllerInfo);
 	}
@@ -84,22 +97,35 @@ void FPlayerControllersDebugBlock::OnDebugDraw(const FCameraDebugBlockDrawParams
 				Renderer.AddText(TEXT("- {cam_notice}%s{cam_default}"), *PlayerController.PlayerControllerName);
 				Renderer.AddIndent();
 				{
+					const FIntPoint ViewportSize = PlayerController.ViewportSize;
+					const double AspectRatio = (ViewportSize.Y != 0 ? (double)ViewportSize.X / (double)ViewportSize.Y : 0.f);
 					Renderer.AddText(TEXT("Local player: {cam_notice}%s{cam_default}\n"), *PlayerController.LocalPlayerName);
 					Renderer.AddText(TEXT("Camera manager: {cam_notice}%s{cam_default}\n"), *PlayerController.CameraManagerName);
+					Renderer.AddText(TEXT("Viewport: {cam_notice}%d x %d{cam_default} (aspect ratio %.3f)\n"), ViewportSize.X, ViewportSize.Y, AspectRatio);
+					if (PlayerController.bHasDefaultAspectRatioAxisConstraint)
+					{
+						Renderer.AddText(TEXT("DefaultAspectRatioAxisConstraint: {cam_notice}%s{cam_default}\n"), *ToDebugString(PlayerController.DefaultAspectRatioAxisConstraint));
+					}
+					else
+					{
+						Renderer.AddText(TEXT("DefaultAspectRatioAxisConstraint: <no local player>\n"));
+					}
+
 					Renderer.AddText(TEXT("View target: {cam_notice}%s{cam_default}"), *PlayerController.ViewTargetName);
 					Renderer.AddIndent();
 					{
-						Renderer.AddText(TEXT("Location  : %s\n"), *ToDebugString(PlayerController.ViewTargetLocation));
-						Renderer.AddText(TEXT("Rotation  : %s\n"), *ToDebugString(PlayerController.ViewTargetRotation));
-						Renderer.AddText(TEXT("FOV  : %s\n"), *ToDebugString(PlayerController.ViewTargetFOV));
-						Renderer.AddText(TEXT("AspectRatio  : %s\n"), *ToDebugString(PlayerController.ViewTargetAspectRatio));
-						if (PlayerController.DefaultAspectRatioAxisConstraint.IsSet())
+						Renderer.AddText(TEXT("Location: %s\n"), *ToDebugString(PlayerController.ViewTargetLocation));
+						Renderer.AddText(TEXT("Rotation: %s\n"), *ToDebugString(PlayerController.ViewTargetRotation));
+						Renderer.AddText(TEXT("FOV: %s\n"), *ToDebugString(PlayerController.ViewTargetFOV));
+						Renderer.AddText(TEXT("AspectRatio: %s\n"), *ToDebugString(PlayerController.ViewTargetAspectRatio));
+						Renderer.AddText(TEXT("ConstrainAspectRatio: %s\n"), *ToDebugString(PlayerController.bViewTargetConstrainAspectRatio));
+						if (PlayerController.bViewTargetOverrideAspectRatioAxisConstraint)
 						{
-							Renderer.AddText(TEXT("DefaultAspectRatioAxisConstraint  : %s\n"), *ToDebugString(PlayerController.DefaultAspectRatioAxisConstraint.GetValue()));
+							Renderer.AddText(TEXT("AspectRatioAxisConstraint: %s\n"), *ToDebugString(PlayerController.ViewTargetAspectRatioAxisConstraint));
 						}
 						else
 						{
-							Renderer.AddText(TEXT("DefaultAspectRatioAxisConstraint  : <no local player>\n"));
+							Renderer.AddText(TEXT("AspectRatioAxisConstraint: <no override>\n"));
 						}
 					}
 					Renderer.RemoveIndent();
@@ -126,10 +152,15 @@ FArchive& operator<< (FArchive& Ar, FPlayerControllersDebugBlock::FPlayerControl
 	Ar << PlayerControllerInfo.PlayerControllerName;
 	Ar << PlayerControllerInfo.CameraManagerName;
 	Ar << PlayerControllerInfo.ViewTargetName;
+	Ar << PlayerControllerInfo.LocalPlayerName;
+	Ar << PlayerControllerInfo.ViewportSize;
+	Ar << PlayerControllerInfo.DefaultAspectRatioAxisConstraint;
 	Ar << PlayerControllerInfo.ViewTargetLocation;
 	Ar << PlayerControllerInfo.ViewTargetRotation;
 	Ar << PlayerControllerInfo.ViewTargetFOV;
 	Ar << PlayerControllerInfo.ViewTargetAspectRatio;
+	Ar << PlayerControllerInfo.ViewTargetAspectRatioAxisConstraint;
+	Ar << PlayerControllerInfo.bViewTargetConstrainAspectRatio;
 	return Ar;
 }
 
diff --git a/Engine/Plugins/Cameras/GameplayCameras/Source/GameplayCameras/Public/Debug/PlayerControllersDebugBlock.h b/Engine/Plugins/Cameras/GameplayCameras/Source/GameplayCameras/Public/Debug/PlayerControllersDebugBlock.h
--- a/Engine/Plugins/Cameras/GameplayCameras/Source/GameplayCameras/Public/Debug/PlayerControllersDebugBlock.h
+++ b/Engine/Plugins/Cameras/GameplayCameras/Source/GameplayCameras/Public/Debug/PlayerControllersDebugBlock.h
@@ -38,13 +38,20 @@ private:
 	{
 		FString PlayerControllerName;
 		FString CameraManagerName;
-		FString ViewTargetName;
 		FString LocalPlayerName;
-		TOptional<EAspectRatioAxisConstraint> DefaultAspectRatioAxisConstraint;
+
+		FIntPoint ViewportSize;
+		TEnumAsByte<EAspectRatioAxisConstraint> DefaultAspectRatioAxisConstraint = AspectRatio_MaintainXFOV;
+		bool bHasDefaultAspectRatioAxisConstraint = true;
+
+		FString ViewTargetName;
 		FVector3d ViewTargetLocation;
 		FRotator3d ViewTargetRotation;
 		float ViewTargetFOV;
 		float ViewTargetAspectRatio;
+		TEnumAsByte<EAspectRatioAxisConstraint> ViewTargetAspectRatioAxisConstraint = AspectRatio_MaintainXFOV;
+		bool bViewTargetConstrainAspectRatio;
+		bool bViewTargetOverrideAspectRatioAxisConstraint = false;
 	};
 	TArray<FPlayerControllerDebugInfo> PlayerControllers;
 	bool bHadValidWorld = false;
