Cameras: don't find referenced objects from too many packages when gathering director usage

The StateTree and Blueprint directors figure out what camera rigs and proxies they use by simply looking at referenced objects up to a certain distance. However, this reference gathering was a bit too heavy handed. This change skips compiled packages, and adds an option to not explore objects of certain classes. This helps drive down the amount of referenced rigs and proxies shown in the "family shortcuts" in the editor (which showed seemingly unrelated assets). It also fixes a bug with camera assets being dirtied on PIE because they would find new rigs to add to the default parameters when upgrading old data to the new interface.

#Synced-CL 49781371

diff --git a/Engine/Plugins/Cameras/GameplayCameras/Source/GameplayCameras/Private/Directors/BlueprintCameraDirector.cpp b/Engine/Plugins/Cameras/GameplayCameras/Source/GameplayCameras/Private/Directors/BlueprintCameraDirector.cpp
--- a/Engine/Plugins/Cameras/GameplayCameras/Source/GameplayCameras/Private/Directors/BlueprintCameraDirector.cpp
+++ b/Engine/Plugins/Cameras/GameplayCameras/Source/GameplayCameras/Private/Directors/BlueprintCameraDirector.cpp
@@ -532,7 +532,10 @@ void UBlueprintCameraDirector::OnGatherRigUsageInfo(FCameraDirectorRigUsageInfo&
 
 	TArray<UClass*> RefClasses { UCameraRigAsset::StaticClass(), UCameraRigProxyAsset::StaticClass() };
 	FOutgoingReferenceFinder ReferenceFinder(EvaluatorBlueprint, RefClasses);
-	ReferenceFinder.SetMaxDistance(GetDefault<UGameplayCamerasSettings>()->MaxUsageSearchDistance);
+	ReferenceFinder.MaxDistance = GetDefault<UGameplayCamerasSettings>()->MaxUsageSearchDistance;
+	// Find camera rigs and proxies, but don't look at *their* outgoing references, otherwise we might pick rig prefabs,
+	// or rigs referenced in transition conditions.
+	ReferenceFinder.ExcludeClasses = RefClasses;
 	ReferenceFinder.CollectReferences();
 	ReferenceFinder.GetReferencesOfClass<UCameraRigAsset>(UsageInfo.CameraRigs);
 	ReferenceFinder.GetReferencesOfClass<UCameraRigProxyAsset>(UsageInfo.CameraRigProxies);
diff --git a/Engine/Plugins/Cameras/GameplayCameras/Source/GameplayCameras/Private/Directors/StateTreeCameraDirector.cpp b/Engine/Plugins/Cameras/GameplayCameras/Source/GameplayCameras/Private/Directors/StateTreeCameraDirector.cpp
--- a/Engine/Plugins/Cameras/GameplayCameras/Source/GameplayCameras/Private/Directors/StateTreeCameraDirector.cpp
+++ b/Engine/Plugins/Cameras/GameplayCameras/Source/GameplayCameras/Private/Directors/StateTreeCameraDirector.cpp
@@ -251,7 +251,10 @@ void UStateTreeCameraDirector::OnGatherRigUsageInfo(FCameraDirectorRigUsageInfo&
 
 	TArray<UClass*> RefClasses { UCameraRigAsset::StaticClass(), UCameraRigProxyAsset::StaticClass() };
 	FOutgoingReferenceFinder ReferenceFinder(const_cast<UStateTree*>(StateTree), RefClasses);
-	ReferenceFinder.SetMaxDistance(GetDefault<UGameplayCamerasSettings>()->MaxUsageSearchDistance);
+	ReferenceFinder.MaxDistance = GetDefault<UGameplayCamerasSettings>()->MaxUsageSearchDistance;
+	// Find camera rigs and proxies, but don't look at *their* outgoing references, otherwise we might pick rig prefabs,
+	// or rigs referenced in transition conditions.
+	ReferenceFinder.ExcludeClasses = RefClasses;
 	ReferenceFinder.CollectReferences();
 	ReferenceFinder.GetReferencesOfClass<UCameraRigAsset>(UsageInfo.CameraRigs);
 	ReferenceFinder.GetReferencesOfClass<UCameraRigProxyAsset>(UsageInfo.CameraRigProxies);
diff --git a/Engine/Plugins/Cameras/GameplayCameras/Source/GameplayCameras/Private/Helpers/OutgoingReferenceFinder.cpp b/Engine/Plugins/Cameras/GameplayCameras/Source/GameplayCameras/Private/Helpers/OutgoingReferenceFinder.cpp
--- a/Engine/Plugins/Cameras/GameplayCameras/Source/GameplayCameras/Private/Helpers/OutgoingReferenceFinder.cpp
+++ b/Engine/Plugins/Cameras/GameplayCameras/Source/GameplayCameras/Private/Helpers/OutgoingReferenceFinder.cpp
@@ -25,6 +25,9 @@ FOutgoingReferenceFinder::FOutgoingReferenceFinder(UObject* InRootObject, TArray
 
 void FOutgoingReferenceFinder::Initialize(UObject* InRootObject)
 {
+	// Skip packages that correspond to C++ modules and other compiled things.
+	ExcludePackageFlags |= PKG_CompiledIn;
+
 	RootObject = InRootObject;
 	PackageScope = InRootObject->GetOutermost();
 
@@ -36,13 +39,10 @@ void FOutgoingReferenceFinder::Initialize(UObject* InRootObject)
 	ArShouldSkipBulkData = true;
 }
 
-void FOutgoingReferenceFinder::SetMaxDistance(int32 InMaxDistance)
-{
-	MaxDistance = FMath::Clamp(InMaxDistance, 0, 5);
-}
-
 void FOutgoingReferenceFinder::CollectReferences()
 {
+	ValidMaxDistance = FMath::Clamp(MaxDistance, 0, 5);
+
 	ObjectsToVisit.Reset();
 	VisitedObjects.Reset();
 
@@ -54,7 +54,7 @@ void FOutgoingReferenceFinder::CollectReferences()
 		{
 			VisitedObjects.Add(CurObj.Obj);
 
-			SerializeState = { CurObj.Package, CurObj.Distance };
+			SerializeState = { CurObj.Obj, CurObj.Package, CurObj.Distance };
 			CurObj.Obj->Serialize(*this);
 			SerializeState = FSerializeState();
 		}
@@ -93,8 +93,11 @@ FArchive& FOutgoingReferenceFinder::operator<<(UObject*& ObjRef)
 			// has a lower distance. But that's OK. Either we queue a visit of ObjRef right now, and we'll ignore the
 			// duplicate with the shorter distance, or the current distance is over the threshold and we skip visiting
 			// it now, but we'll visit it when we encounter it again with a distance under the threshold.
-			const int32 ObjRefDistance = SerializeState.CurrentDistance + (ObjRefPackage == SerializeState.CurrentPackage ? 0 : 1);
-			if (ObjRefDistance <= MaxDistance)
+			const int32 AddOneDistance = (ObjRefPackage == SerializeState.CurrentPackage ? 0 : 1);
+			const int32 ObjRefDistance = SerializeState.CurrentDistance + AddOneDistance;
+			const bool bExcludePackage = ((ObjRefPackage->GetPackageFlags() & ExcludePackageFlags) != 0);
+			const bool bExcludeObject = IsIgnoredClass(ObjClass);
+			if (ObjRefDistance <= ValidMaxDistance && !bExcludePackage && !bExcludeObject)
 			{
 				ObjectsToVisit.Add({ ObjRef, ObjRefPackage, ObjRefDistance });
 			}
@@ -116,5 +119,14 @@ bool FOutgoingReferenceFinder::MatchesAnyTargetClass(UClass* InObjClass) const
 	return false;
 }
 
+bool FOutgoingReferenceFinder::IsIgnoredClass(UClass* InObjClass) const
+{
+	return ExcludeClasses.ContainsByPredicate(
+			[InObjClass](UClass* ExcludeClass)
+			{
+				return ExcludeClass ? InObjClass->IsChildOf(ExcludeClass) : false;
+			});
+}
+
 }  // namespace UE::Cameras
 
diff --git a/Engine/Plugins/Cameras/GameplayCameras/Source/GameplayCameras/Public/Helpers/OutgoingReferenceFinder.h b/Engine/Plugins/Cameras/GameplayCameras/Source/GameplayCameras/Public/Helpers/OutgoingReferenceFinder.h
--- a/Engine/Plugins/Cameras/GameplayCameras/Source/GameplayCameras/Public/Helpers/OutgoingReferenceFinder.h
+++ b/Engine/Plugins/Cameras/GameplayCameras/Source/GameplayCameras/Public/Helpers/OutgoingReferenceFinder.h
@@ -7,6 +7,7 @@
 #include "Containers/Set.h"
 #include "CoreTypes.h"
 #include "Serialization/ArchiveUObject.h"
+#include "UObject/Class.h"
 
 #define UE_API GAMEPLAYCAMERAS_API
 
@@ -21,23 +22,32 @@ class FOutgoingReferenceFinder : public FArchiveUObject
 {
 public:
 
-	/** Creates a new instance of the reference finder. */
-	UE_API FOutgoingReferenceFinder(UObject* InRootObject, UClass* InReferencedObjectClass);
-
-	/** Creates a new instance of the reference finder. */
-	UE_API FOutgoingReferenceFinder(UObject* InRootObject, TArrayView<UClass*> InReferencedObjectClasses);
-
 	/**
-	 * Sets the maximum distance (in number of packages) to search for references.
+	 * The maximum distance (in number of packages) to search for references.
 	 * Zero means only search the root object package.
 	 * One means search directly referenced packages.
 	 * Higher distances exponentially grow the number of objects to search, so the maximum
 	 * value settable is 5.
 	 */
-	UE_API void SetMaxDistance(int32 InMaxDistance);
+	int32 MaxDistance = 1;
 
-	/** Gets the maximum distance (in number of packages) to search for references. */
-	int32 GetMaxDistance() const { return MaxDistance; }
+	/**
+	 * EPackageFlags that determine what packages to skip.
+	 */
+	uint32 ExcludePackageFlags = 0;
+
+	/**
+	 * Classes to not inspect for outgoing references.
+	 */
+	TArray<UClass*> ExcludeClasses;
+
+public:
+
+	/** Creates a new instance of the reference finder. */
+	UE_API FOutgoingReferenceFinder(UObject* InRootObject, UClass* InReferencedObjectClass);
+
+	/** Creates a new instance of the reference finder. */
+	UE_API FOutgoingReferenceFinder(UObject* InRootObject, TArrayView<UClass*> InReferencedObjectClasses);
 
 	/** Runs the reference finding. */
 	UE_API void CollectReferences();
@@ -56,9 +66,10 @@ protected:
 
 private:
 
-	UE_API void Initialize(UObject* InRootObject);
+	void Initialize(UObject* InRootObject);
 
-	UE_API bool MatchesAnyTargetClass(UClass* InObjClass) const;
+	bool MatchesAnyTargetClass(UClass* InObjClass) const;
+	bool IsIgnoredClass(UClass* InObjClass) const;
 
 private:
 
@@ -71,13 +82,14 @@ private:
 
 	struct FSerializeState
 	{
+		UObject* CurrentObject = nullptr;
 		UPackage* CurrentPackage = nullptr;
 		int32 CurrentDistance = 0;
 	};
 
 	UObject* RootObject;
 	UPackage* PackageScope;
-	int32 MaxDistance = 1;
+	int32 ValidMaxDistance = 1;
 
 	TSet<UClass*> TargetObjectClasses;
 
@@ -91,15 +103,20 @@ private:
 template<typename ObjectClass>
 bool FOutgoingReferenceFinder::GetReferencesOfClass(TArray<ObjectClass*>& OutReferencedObjects) const
 {
-	if (const TSet<UObject*>* ReferencesOfClass = ReferencedObjects.Find(ObjectClass::StaticClass()))
+	bool bHadAny = false;
+	UClass* ObjectClassClass = ObjectClass::StaticClass();
+	for (const auto& Pair : ReferencedObjects)
 	{
-		for (UObject* Obj : (*ReferencesOfClass))
+		if (Pair.Key->IsChildOf(ObjectClassClass))
 		{
-			OutReferencedObjects.Add(CastChecked<ObjectClass>(Obj));
+			bHadAny |= (Pair.Value.Num() > 0);
+			for (UObject* Obj : Pair.Value)
+			{
+				OutReferencedObjects.Add(CastChecked<ObjectClass>(Obj));
+			}
 		}
-		return !ReferencesOfClass->IsEmpty();
 	}
-	return false;
+	return bHadAny;
 }
 
 }  // namespace UE::Cameras
