Cameras: add a time-out property to camera actions

Also simplify a bit how actions get cleaned-up, and the properties that need serialization.

#Synced-CL 49781371

diff --git a/Engine/Plugins/Cameras/GameplayCameras/Source/GameplayCameras/Private/Actions/BaseAimAtCameraAction.cpp b/Engine/Plugins/Cameras/GameplayCameras/Source/GameplayCameras/Private/Actions/BaseAimAtCameraAction.cpp
--- a/Engine/Plugins/Cameras/GameplayCameras/Source/GameplayCameras/Private/Actions/BaseAimAtCameraAction.cpp
+++ b/Engine/Plugins/Cameras/GameplayCameras/Source/GameplayCameras/Private/Actions/BaseAimAtCameraAction.cpp
@@ -141,6 +141,7 @@ bool FBaseAimAtCameraActionEvaluator::ComputeDesiredCorrection(const FCameraActi
 
 	static const FVector2d CenterFraming(0.5, 0.5);
 	const bool bNeedsTargetFraming = (TargetFraming != CenterFraming);
+	FVector3d TargetFramingAim = ScratchResult.CameraPose.GetAimDir();
 	if (bNeedsTargetFraming)
 	{
 		TSharedPtr<const FCameraEvaluationContext> EvaluationContext = Params.Scope->GetEvaluationContext();
@@ -148,12 +149,9 @@ bool FBaseAimAtCameraActionEvaluator::ComputeDesiredCorrection(const FCameraActi
 				ScratchResult.CameraPose, EvaluationContext, TargetFraming);
 		TargetFramingAim = TargetRay.Direction;
 	}
-	else
-	{
-		TargetFramingAim = ScratchResult.CameraPose.GetAimDir();
-	}
 
 #if UE_GAMEPLAY_CAMERAS_DEBUG
+	DebugTargetFramingAim = TargetFramingAim;
 	DebugCameraLocation = ScratchResult.CameraPose.GetLocation();
 	DebugCameraAim = ScratchResult.CameraPose.GetAimDir();
 #endif
@@ -308,13 +306,11 @@ void FBaseAimAtCameraActionEvaluator::OnSerialize(const FCameraActionEvaluatorSe
 		Interpolator->Serialize(InterpolatorParams, Ar);
 	}
 
-	Ar << TargetLocation;
-
-	LastCameraPose.SerializeWithFlags(Ar);
-	Ar << LastContextLocation;
-	Ar << LastPivotLocation;
-
 	Ar << bIsLockedOn;
+
+	// Other fields don't need to be serialized since they are either recomputed on action clone and then never change
+	// (e.g. ChildHierarchy), or recomputed every frame (e.g. TargetLocation, TargetFraming, and the stuff for running
+	// the temporary IK evaluation).
 }
 
 #if UE_GAMEPLAY_CAMERAS_DEBUG
@@ -337,7 +333,7 @@ void FBaseAimAtCameraActionEvaluator::OnBuildDebugBlocks(const FCameraDebugBlock
 	DebugBlock.CorrectionLeft = DebugCorrectionLeft;
 	DebugBlock.CurrentCorrection = DebugCurrentCorrection;
 	DebugBlock.TargetFraming = TargetFraming;
-	DebugBlock.TargetFramingAim = TargetFramingAim;
+	DebugBlock.TargetFramingAim = DebugTargetFramingAim;
 	DebugBlock.CameraLocation = DebugCameraLocation;
 	DebugBlock.CameraAim = DebugCameraAim;
 	DebugBlock.bIsLockedOn = bIsLockedOn;
diff --git a/Engine/Plugins/Cameras/GameplayCameras/Source/GameplayCameras/Private/Services/CameraActionScope.cpp b/Engine/Plugins/Cameras/GameplayCameras/Source/GameplayCameras/Private/Services/CameraActionScope.cpp
--- a/Engine/Plugins/Cameras/GameplayCameras/Source/GameplayCameras/Private/Services/CameraActionScope.cpp
+++ b/Engine/Plugins/Cameras/GameplayCameras/Source/GameplayCameras/Private/Services/CameraActionScope.cpp
@@ -27,7 +27,6 @@ FCameraActionInstanceID FCameraActionScope::StartAction(const UCameraAction* InC
 			FActionInfo& NewAction = Actions.Emplace_GetRef();
 			NewAction.Evaluator = Evaluator;
 			NewAction.InstanceID = FCameraActionInstanceID(NextActionID++);
-			NewAction.State = EActionState::Uninitialized;
 			return NewAction.InstanceID;
 		}
 	}
@@ -64,18 +63,26 @@ bool FCameraActionScope::StopAction(const FCameraActionInstanceID InInstanceID)
 	if (Index != INDEX_NONE)
 	{
 		FActionInfo& ActionInfo = Actions[Index];
-		if (ActionInfo.State == EActionState::Running)
-		{
-			FCameraActionEvaluatorTeardownParams Params;
-			Params.Scope = SharedThis(this);
-			ActionInfo.Evaluator->Teardown(Params);
-		}
+		StopAction(ActionInfo);
 		Actions.RemoveAt(Index);
 		return true;
 	}
 	return false;
 }
 
+void FCameraActionScope::StopAction(FActionInfo& ActionInfo)
+{
+	// Only teardown the evaluator if it was initialized in the first place.
+	if (ActionInfo.State == EActionState::Running && ActionInfo.Evaluator.IsValid())
+	{
+		FCameraActionEvaluatorTeardownParams Params;
+		Params.Scope = SharedThis(this);
+		ActionInfo.Evaluator->Teardown(Params);
+	}
+	ActionInfo.Evaluator.Reset();
+	ActionInfo.State = EActionState::Finished;
+}
+
 void FCameraActionScope::Initialize(const FCameraRigEvaluationInfo& CameraRigInfo)
 {
 	// Don't store CameraRigInfo directly because we want to keep the evaluation context
@@ -119,6 +126,16 @@ void FCameraActionScope::PreScopeRun(const FCameraNodeEvaluationParams& Params,
 		FCameraActionEvaluationResult ActionResult(OutResult);
 		if (!PrepareActionForRun(ActionInfo, ActionResult))
 		{
+			StopAction(ActionInfo);
+			It.RemoveCurrent();
+			continue;
+		}
+
+		// Only check for time-out on pre-scope since we don't want to update ElapsedTime twice in a frame.
+		if (HasActionTimedOut(Params.DeltaTime, ActionInfo))
+		{
+			StopAction(ActionInfo);
+			It.RemoveCurrent();
 			continue;
 		}
 
@@ -144,6 +161,8 @@ void FCameraActionScope::PostScopeRun(const FCameraNodeEvaluationParams& Params,
 		FCameraActionEvaluationResult ActionResult(OutResult);
 		if (!PrepareActionForRun(ActionInfo, ActionResult))
 		{
+			StopAction(ActionInfo);
+			It.RemoveCurrent();
 			continue;
 		}
 
@@ -177,6 +196,17 @@ bool FCameraActionScope::PrepareActionForRun(FActionInfo& ActionInfo, FCameraAct
 	return ActionInfo.State == EActionState::Running;
 }
 
+bool FCameraActionScope::HasActionTimedOut(float DeltaTime, FActionInfo& ActionInfo)
+{
+	ActionInfo.ElapsedTime += DeltaTime;
+	const UCameraAction* ActionData = ActionInfo.Evaluator->GetCameraAction();
+	if (ensure(ActionData))
+	{
+		return (ActionData->TimeOut > 0.f && ActionInfo.ElapsedTime >= ActionData->TimeOut);
+	}
+	return true;
+}
+
 void FCameraActionScope::Serialize(const FCameraActionEvaluatorSerializeParams& Params, FArchive& Ar)
 {
 	int32 NumActionsToSerialize = Actions.Num();
diff --git a/Engine/Plugins/Cameras/GameplayCameras/Source/GameplayCameras/Public/Actions/BaseAimAtCameraAction.h b/Engine/Plugins/Cameras/GameplayCameras/Source/GameplayCameras/Public/Actions/BaseAimAtCameraAction.h
--- a/Engine/Plugins/Cameras/GameplayCameras/Source/GameplayCameras/Public/Actions/BaseAimAtCameraAction.h
+++ b/Engine/Plugins/Cameras/GameplayCameras/Source/GameplayCameras/Public/Actions/BaseAimAtCameraAction.h
@@ -92,11 +92,6 @@ private:
 
 	FVector3d TargetLocation;
 	FVector2d TargetFraming = { 0.5, 0.5 };
-	FVector3d TargetFramingAim;
-
-	FCameraPose LastCameraPose;
-	FVector3d LastContextLocation;
-	FVector3d LastPivotLocation;
 
 	bool bIsLockedOn = false;
 
@@ -104,6 +99,7 @@ private:
 	float DebugElapsedTime = 0.f;
 	FRotator3d DebugCorrectionLeft;
 	FRotator3d DebugCurrentCorrection;
+	FVector3d DebugTargetFramingAim;
 	FVector3d DebugCameraLocation;
 	FVector3d DebugCameraAim;
 #endif  // UE_GAMEPLAY_CAMERAS_DEBUG
diff --git a/Engine/Plugins/Cameras/GameplayCameras/Source/GameplayCameras/Public/Services/CameraAction.h b/Engine/Plugins/Cameras/GameplayCameras/Source/GameplayCameras/Public/Services/CameraAction.h
--- a/Engine/Plugins/Cameras/GameplayCameras/Source/GameplayCameras/Public/Services/CameraAction.h
+++ b/Engine/Plugins/Cameras/GameplayCameras/Source/GameplayCameras/Public/Services/CameraAction.h
@@ -33,6 +33,13 @@ public:
 	 */
 	UPROPERTY(EditAnywhere, BlueprintReadWrite, Category=Common, meta=(ExposeOnSpawn=true))
 	bool bPropagateToNewCameraRigs = false;
+
+	/**
+	 * The time, in seconds, after which the action should be automatically stopped. Zero or negative time means that
+	 * this action never times out (it either completes on its own, or gets stopped by the user).
+	 */
+	UPROPERTY(EditAnywhere, BlueprintReadWrite, Category=Common, meta=(ExposeOnSpawn=true))
+	float TimeOut = 0.f;
 	
 public:
 
diff --git a/Engine/Plugins/Cameras/GameplayCameras/Source/GameplayCameras/Public/Services/CameraActionScope.h b/Engine/Plugins/Cameras/GameplayCameras/Source/GameplayCameras/Public/Services/CameraActionScope.h
--- a/Engine/Plugins/Cameras/GameplayCameras/Source/GameplayCameras/Public/Services/CameraActionScope.h
+++ b/Engine/Plugins/Cameras/GameplayCameras/Source/GameplayCameras/Public/Services/CameraActionScope.h
@@ -93,6 +93,8 @@ private:
 	struct FActionInfo;
 
 	bool PrepareActionForRun(FActionInfo& ActionInfo, FCameraActionEvaluationResult& OutResult);
+	bool HasActionTimedOut(float DeltaTime, FActionInfo& ActionInfo);
+	void StopAction(FActionInfo& ActionInfo);
 
 private:
 
@@ -107,6 +109,7 @@ private:
 	{
 		TSharedPtr<FCameraActionEvaluator> Evaluator;
 		FCameraActionInstanceID InstanceID;
+		float ElapsedTime = 0.f;
 		EActionState State = EActionState::Uninitialized;
 	};
 
