Cameras: add action that aims at an actor with optional bone/socket

#Synced-CL 49781371

diff --git a/Engine/Plugins/Cameras/GameplayCameras/Source/GameplayCameras/Private/Actions/AimAtActorCameraAction.cpp b/Engine/Plugins/Cameras/GameplayCameras/Source/GameplayCameras/Private/Actions/AimAtActorCameraAction.cpp
new file mode 100644
--- /dev/null
+++ b/Engine/Plugins/Cameras/GameplayCameras/Source/GameplayCameras/Private/Actions/AimAtActorCameraAction.cpp
@@ -0,0 +1,76 @@
+// Copyright Epic Games, Inc. All Rights Reserved.
+
+#include "Actions/AimAtActorCameraAction.h"
+
+#include "Components/SkeletalMeshComponent.h"
+#include "GameFramework/Actor.h"
+
+#include UE_INLINE_GENERATED_CPP_BY_NAME(AimAtActorCameraAction)
+
+namespace UE::Cameras
+{
+
+class FAimAtActorCameraActionEvaluator : public FBaseAimAtCameraActionEvaluator
+{
+	UE_DECLARE_CAMERA_ACTION_EVALUATOR_EX(, FAimAtActorCameraActionEvaluator, FBaseAimAtCameraActionEvaluator)
+
+protected:
+
+	// FBaseAimAtCameraActionEvaluator
+	virtual FVector3d UpdateTargetLocation(const FCameraActionEvaluationParams& Params, const FCameraActionEvaluationResult& Result) override;
+
+private:
+
+	AActor* CachedActor = nullptr;
+	USkeletalMeshComponent* CachedSkeletalMeshComponent = nullptr;
+	FName CachedBoneName;
+};
+
+UE_DEFINE_CAMERA_ACTION_EVALUATOR(FAimAtActorCameraActionEvaluator)
+
+FVector3d FAimAtActorCameraActionEvaluator::UpdateTargetLocation(const FCameraActionEvaluationParams& Params, const FCameraActionEvaluationResult& Result)
+{
+	const UAimAtActorCameraAction* ActionData = GetCameraActionAs<UAimAtActorCameraAction>();
+	
+	const bool bIsSameActor = (ActionData->TargetActor == CachedActor);
+	if (!bIsSameActor)
+	{
+		CachedActor = ActionData->TargetActor.Get();
+
+		CachedSkeletalMeshComponent = nullptr;
+		if (ActionData->TargetActor && (!ActionData->TargetSocketName.IsNone() || !ActionData->TargetBoneName.IsNone()))
+		{
+			CachedSkeletalMeshComponent = ActionData->TargetActor->FindComponentByClass<USkeletalMeshComponent>();
+		}
+
+		CachedBoneName = NAME_None;
+		if (CachedSkeletalMeshComponent)
+		{
+			CachedBoneName = ActionData->TargetBoneName;
+			if (!ActionData->TargetSocketName.IsNone())
+			{
+				CachedBoneName = CachedSkeletalMeshComponent->GetSocketBoneName(ActionData->TargetSocketName);
+			}
+		}
+	}
+
+	if (CachedSkeletalMeshComponent && !CachedBoneName.IsNone())
+	{
+		return CachedSkeletalMeshComponent->GetBoneTransform(CachedBoneName).GetLocation();
+	}
+	else if (CachedActor)
+	{
+		return CachedActor->GetTransform().GetLocation();
+	}
+
+	return FVector3d::ZeroVector;
+}
+
+}  // namespace UE::Cameras
+
+FCameraActionEvaluatorPtr UAimAtActorCameraAction::OnBuildEvaluator(FCameraActionEvaluatorBuilder& Builder) const
+{
+	using namespace UE::Cameras;
+	return Builder.BuildEvaluator<FAimAtActorCameraActionEvaluator>();
+}
+
diff --git a/Engine/Plugins/Cameras/GameplayCameras/Source/GameplayCameras/Public/Actions/AimAtActorCameraAction.h b/Engine/Plugins/Cameras/GameplayCameras/Source/GameplayCameras/Public/Actions/AimAtActorCameraAction.h
new file mode 100644
--- /dev/null
+++ b/Engine/Plugins/Cameras/GameplayCameras/Source/GameplayCameras/Public/Actions/AimAtActorCameraAction.h
@@ -0,0 +1,36 @@
+// Copyright Epic Games, Inc. All Rights Reserved.
+
+#pragma once
+
+#include "Actions/BaseAimAtCameraAction.h"
+
+#include "AimAtActorCameraAction.generated.h"
+
+/**
+ * A camera action that aims at a given target, using a given interpolation.
+ */
+UCLASS()
+class UAimAtActorCameraAction : public UBaseAimAtCameraAction
+{
+	GENERATED_BODY()
+
+public:
+
+	/** The actor to attach to. */
+	UPROPERTY(BlueprintReadWrite, EditAnywhere, Category="Aiming", meta=(ExposeOnSpawn=true))
+	TObjectPtr<AActor> TargetActor;
+
+	/** An optional socket to attach to on the actor. */
+	UPROPERTY(BlueprintReadWrite, EditAnywhere, Interp, Category="Aiming", meta=(ExposeOnSpawn=true))
+	FName TargetSocketName;
+
+	/** An optional bone to attach to on the actor. */
+	UPROPERTY(BlueprintReadWrite, EditAnywhere, Interp, Category="Aiming", meta=(ExposeOnSpawn=true))
+	FName TargetBoneName;
+
+protected:
+
+	// UCameraAction interface.
+	virtual FCameraActionEvaluatorPtr OnBuildEvaluator(FCameraActionEvaluatorBuilder& Builder) const override;
+};
+
