Cameras: export more things out of the module
#Synced-CL 51024716

diff --git a/Engine/Plugins/Cameras/GameplayCameras/Source/GameplayCameras/Public/GameFramework/GameplayCameraActorBase.h b/Engine/Plugins/Cameras/GameplayCameras/Source/GameplayCameras/Public/GameFramework/GameplayCameraActorBase.h
--- a/Engine/Plugins/Cameras/GameplayCameras/Source/GameplayCameras/Public/GameFramework/GameplayCameraActorBase.h
+++ b/Engine/Plugins/Cameras/GameplayCameras/Source/GameplayCameras/Public/GameFramework/GameplayCameraActorBase.h
@@ -24,7 +24,7 @@ class AGameplayCameraActorBase : public AActor
 
 public:
 
-	AGameplayCameraActorBase(const FObjectInitializer& ObjectInit);
+	GAMEPLAYCAMERAS_API AGameplayCameraActorBase(const FObjectInitializer& ObjectInit);
 
 public:
 
diff --git a/Engine/Plugins/Cameras/GameplayCameras/Source/GameplayCameras/Public/GameFramework/GameplayCameraComponentBase.h b/Engine/Plugins/Cameras/GameplayCameras/Source/GameplayCameras/Public/GameFramework/GameplayCameraComponentBase.h
--- a/Engine/Plugins/Cameras/GameplayCameras/Source/GameplayCameras/Public/GameFramework/GameplayCameraComponentBase.h
+++ b/Engine/Plugins/Cameras/GameplayCameras/Source/GameplayCameras/Public/GameFramework/GameplayCameraComponentBase.h
@@ -103,39 +103,39 @@ public:
 
 	/** Activates the given camera rig prefab in the base layer. */
 	UFUNCTION(BlueprintCallable, Category="Camera")
-	FCameraRigInstanceID ActivatePersistentBaseCameraRig(UCameraRigAsset* CameraRig);
+	GAMEPLAYCAMERAS_API FCameraRigInstanceID ActivatePersistentBaseCameraRig(UCameraRigAsset* CameraRig);
 
 	/** Activates the given camera rig prefab in the global layer. */
 	UFUNCTION(BlueprintCallable, Category="Camera")
-	FCameraRigInstanceID ActivatePersistentGlobalCameraRig(UCameraRigAsset* CameraRig);
+	GAMEPLAYCAMERAS_API FCameraRigInstanceID ActivatePersistentGlobalCameraRig(UCameraRigAsset* CameraRig);
 
 	/** Activates the given camera rig prefab in the visual layer. */
 	UFUNCTION(BlueprintCallable, Category="Camera")
-	FCameraRigInstanceID ActivatePersistentVisualCameraRig(UCameraRigAsset* CameraRig);
+	GAMEPLAYCAMERAS_API FCameraRigInstanceID ActivatePersistentVisualCameraRig(UCameraRigAsset* CameraRig);
 
 	/** Deactivates a previously activated base, global, or visual camera rig. */
 	UFUNCTION(BlueprintCallable, Category="Camera")
-	void DeactivateCameraRig(FCameraRigInstanceID InstanceID, bool bImmediately = false);
+	GAMEPLAYCAMERAS_API void DeactivateCameraRig(FCameraRigInstanceID InstanceID, bool bImmediately = false);
 
 public:
 
 	/** Starts a camera modifier rig on the global layer. */
 	UFUNCTION(BlueprintCallable, Category="Camera")
-	FCameraRigInstanceID StartGlobalCameraModifierRig(const UCameraRigAsset* CameraRig, int32 OrderKey = 0);
+	GAMEPLAYCAMERAS_API FCameraRigInstanceID StartGlobalCameraModifierRig(const UCameraRigAsset* CameraRig, int32 OrderKey = 0);
 
 	/** Starts a camera modifier rig on the visual layer. */
 	UFUNCTION(BlueprintCallable, Category="Camera")
-	FCameraRigInstanceID StartVisualCameraModifierRig(const UCameraRigAsset* CameraRig, int32 OrderKey = 0);
+	GAMEPLAYCAMERAS_API FCameraRigInstanceID StartVisualCameraModifierRig(const UCameraRigAsset* CameraRig, int32 OrderKey = 0);
 
 	/** Stops a camera modifier rig on previously started on the global or visual layer. */
 	UFUNCTION(BlueprintCallable, Category="Camera")
-	void StopCameraModifierRig(FCameraRigInstanceID InstanceID, bool bImmediately = false);
+	GAMEPLAYCAMERAS_API void StopCameraModifierRig(FCameraRigInstanceID InstanceID, bool bImmediately = false);
 
 public:
 
 	/** Starts a new camera shake. */
 	UFUNCTION(BlueprintCallable, Category="Camera")
-	FCameraShakeInstanceID StartCameraShakeAsset(
+	GAMEPLAYCAMERAS_API FCameraShakeInstanceID StartCameraShakeAsset(
 			const UCameraShakeAsset* CameraShake,
 			float ShakeScale = 1.f,
 			ECameraShakePlaySpace PlaySpace = ECameraShakePlaySpace::CameraLocal,
@@ -143,54 +143,54 @@ public:
 
 	/** Checks if a camera shake is running. */
 	UFUNCTION(BlueprintCallable, Category="Camera")
-	bool IsCameraShakeAssetPlaying(FCameraShakeInstanceID InInstanceID) const;
+	GAMEPLAYCAMERAS_API bool IsCameraShakeAssetPlaying(FCameraShakeInstanceID InInstanceID) const;
 
 	/** Stops a running camera shake. */
 	UFUNCTION(BlueprintCallable, Category="Camera")
-	bool StopCameraShakeAsset(FCameraShakeInstanceID InInstanceID, bool bImmediately = false);
+	GAMEPLAYCAMERAS_API bool StopCameraShakeAsset(FCameraShakeInstanceID InInstanceID, bool bImmediately = false);
 
 public:
 
 	/** Starts the given camera action. */
 	UFUNCTION(BlueprintCallable, Category="Camera")
-	FCameraActionInstanceID StartAction(const UCameraAction* CameraAction);
+	GAMEPLAYCAMERAS_API FCameraActionInstanceID StartAction(const UCameraAction* CameraAction);
 
 	/** Returns whether the given camera action instance is still running. */
 	UFUNCTION(BlueprintCallable, Category="Camera")
-	bool IsActionRunning(FCameraActionInstanceID InInstanceID);
+	GAMEPLAYCAMERAS_API bool IsActionRunning(FCameraActionInstanceID InInstanceID);
 
 	/** Stops the given camera action instance. */
 	UFUNCTION(BlueprintCallable, Category="Camera")
-	bool StopAction(FCameraActionInstanceID InInstanceID);
+	GAMEPLAYCAMERAS_API bool StopAction(FCameraActionInstanceID InInstanceID);
 
 	/** Stops all camera actions of a given class. */
 	UFUNCTION(BlueprintCallable, Category="Camera")
-	bool StopAllActionsOfClass(TSubclassOf<UCameraAction> InActionClass);
+	GAMEPLAYCAMERAS_API bool StopAllActionsOfClass(TSubclassOf<UCameraAction> InActionClass);
 
 public:
 
 	// UActorComponent interface
-	virtual void OnRegister() override;
-	virtual void BeginPlay() override;
-	virtual void EndPlay(const EEndPlayReason::Type EndPlayReason) override;
-	virtual void OnComponentDestroyed(bool bDestroyingHierarchy) override;
-	virtual void Activate(bool bReset = false) override;
-	virtual void Deactivate() override;
-	virtual void TickComponent(float DeltaTime, ELevelTick TickType, FActorComponentTickFunction *ThisTickFunction) override;
+	GAMEPLAYCAMERAS_API virtual void OnRegister() override;
+	GAMEPLAYCAMERAS_API virtual void BeginPlay() override;
+	GAMEPLAYCAMERAS_API virtual void EndPlay(const EEndPlayReason::Type EndPlayReason) override;
+	GAMEPLAYCAMERAS_API virtual void OnComponentDestroyed(bool bDestroyingHierarchy) override;
+	GAMEPLAYCAMERAS_API virtual void Activate(bool bReset = false) override;
+	GAMEPLAYCAMERAS_API virtual void Deactivate() override;
+	GAMEPLAYCAMERAS_API virtual void TickComponent(float DeltaTime, ELevelTick TickType, FActorComponentTickFunction *ThisTickFunction) override;
 #if WITH_EDITOR
-	virtual bool GetEditorPreviewInfo(float DeltaTime, FMinimalViewInfo& ViewOut) override;
+	GAMEPLAYCAMERAS_API virtual bool GetEditorPreviewInfo(float DeltaTime, FMinimalViewInfo& ViewOut) override;
 #endif 
 
 	// USceneComponent interface.
-	virtual void OnUpdateTransform(EUpdateTransformFlags UpdateTransformFlags, ETeleportType Teleport) override;
+	GAMEPLAYCAMERAS_API virtual void OnUpdateTransform(EUpdateTransformFlags UpdateTransformFlags, ETeleportType Teleport) override;
 
 	// UObject interface.
-	virtual void BeginDestroy() override;
+	GAMEPLAYCAMERAS_API virtual void BeginDestroy() override;
 #if WITH_EDITOR
-	virtual void PostEditChangeProperty( struct FPropertyChangedEvent& PropertyChangedEvent) override;
+	GAMEPLAYCAMERAS_API virtual void PostEditChangeProperty( struct FPropertyChangedEvent& PropertyChangedEvent) override;
 #endif
 
-	static void AddReferencedObjects(UObject* InThis, FReferenceCollector& Collector);
+	GAMEPLAYCAMERAS_API static void AddReferencedObjects(UObject* InThis, FReferenceCollector& Collector);
 
 	// IGameplayCameraSystemHost interface.
 	virtual UObject* GetAsObject() override { return this; }
