Cameras: add option to run camera directors and Blueprint camera nodes in editor

#Synced-CL 51024716

diff --git a/Engine/Plugins/Cameras/GameplayCameras/Source/GameplayCameras/Private/Core/CameraEvaluationContext.cpp b/Engine/Plugins/Cameras/GameplayCameras/Source/GameplayCameras/Private/Core/CameraEvaluationContext.cpp
--- a/Engine/Plugins/Cameras/GameplayCameras/Source/GameplayCameras/Private/Core/CameraEvaluationContext.cpp
+++ b/Engine/Plugins/Cameras/GameplayCameras/Source/GameplayCameras/Private/Core/CameraEvaluationContext.cpp
@@ -205,25 +205,48 @@ void FCameraEvaluationContext::AutoCreateDirectorEvaluator()
 
 void FCameraEvaluationContext::AutoCreateEditorPreviewDirectorEvaluator(const FCameraEvaluationContextActivateParams& Params)
 {
-	if (DirectorEvaluator == nullptr &&
-			ensure(Params.Evaluator) && 
-			Params.Evaluator->GetRole() == ECameraSystemEvaluatorRole::EditorPreview)
+	// If a sub-class already created an evaluator, just use that.
+	if (DirectorEvaluator)
 	{
-		FCameraDirectorRigUsageInfo UsageInfo;
+		return;
+	}
 
-		if (CameraAsset && CameraAsset->GetCameraDirector())
-		{
-			UCameraDirector* CameraDirector = CameraAsset->GetCameraDirector();
-			CameraDirector->GatherRigUsageInfo(UsageInfo);
-		}
+	// If we're not meant for an editor preview, bail out.
+	if (!ensure(Params.Evaluator))
+	{
+		return;
+	}
+	if (Params.Evaluator->GetRole() != ECameraSystemEvaluatorRole::EditorPreview)
+	{
+		return;
+	}
 
-		FCameraDirectorEvaluatorBuilder DirectorBuilder(DirectorEvaluatorStorage);
-		DirectorEvaluator = DirectorBuilder.BuildEvaluator<FEditorPreviewCameraDirectorEvaluator>(UsageInfo.CameraRigs);
+	// Find the camera director to use. If it is OK to run in the editor, bail out and let AutoCreateDirectorEvaluator
+	// create it.
+	UCameraDirector* CameraDirector = nullptr;
+	if (CameraAsset && CameraAsset->GetCameraDirector())
+	{
+		CameraDirector = CameraAsset->GetCameraDirector();
+	}
+	
+	if (CameraDirector && CameraDirector->bRunInEditor)
+	{
+		return;
+	}
 
-		FCameraDirectorInitializeParams InitParams;
-		InitParams.OwnerContext = SharedThis(this);
-		DirectorEvaluator->Initialize(InitParams);
+	// Create a simple preview director evaluator that just lets the user switch between the known rigs.
+	FCameraDirectorRigUsageInfo UsageInfo;
+	if (CameraDirector)
+	{
+		CameraDirector->GatherRigUsageInfo(UsageInfo);
 	}
+
+	FCameraDirectorEvaluatorBuilder DirectorBuilder(DirectorEvaluatorStorage);
+	DirectorEvaluator = DirectorBuilder.BuildEvaluator<FEditorPreviewCameraDirectorEvaluator>(UsageInfo.CameraRigs);
+
+	FCameraDirectorInitializeParams InitParams;
+	InitParams.OwnerContext = SharedThis(this);
+	DirectorEvaluator->Initialize(InitParams);
 }
 
 void FCameraEvaluationContext::SetEditorPreviewCameraRigIndex(int32 Index)
diff --git a/Engine/Plugins/Cameras/GameplayCameras/Source/GameplayCameras/Private/Nodes/Utility/BlueprintCameraNode.cpp b/Engine/Plugins/Cameras/GameplayCameras/Source/GameplayCameras/Private/Nodes/Utility/BlueprintCameraNode.cpp
--- a/Engine/Plugins/Cameras/GameplayCameras/Source/GameplayCameras/Private/Nodes/Utility/BlueprintCameraNode.cpp
+++ b/Engine/Plugins/Cameras/GameplayCameras/Source/GameplayCameras/Private/Nodes/Utility/BlueprintCameraNode.cpp
@@ -73,7 +73,7 @@ void FBlueprintCameraNodeEvaluator::OnInitialize(const FCameraNodeEvaluatorIniti
 	}
 
 #if WITH_EDITOR
-	if (Params.Evaluator->GetRole() == ECameraSystemEvaluatorRole::EditorPreview)
+	if (!BlueprintNode->bRunInEditor && Params.Evaluator->GetRole() == ECameraSystemEvaluatorRole::EditorPreview)
 	{
 		return;
 	}
diff --git a/Engine/Plugins/Cameras/GameplayCameras/Source/GameplayCameras/Public/Core/CameraDirector.h b/Engine/Plugins/Cameras/GameplayCameras/Source/GameplayCameras/Public/Core/CameraDirector.h
--- a/Engine/Plugins/Cameras/GameplayCameras/Source/GameplayCameras/Public/Core/CameraDirector.h
+++ b/Engine/Plugins/Cameras/GameplayCameras/Source/GameplayCameras/Public/Core/CameraDirector.h
@@ -102,6 +102,12 @@ public:
 	UPROPERTY(EditAnywhere, Category="Evaluation")
 	FCameraRigProxyRedirectTable CameraRigProxyRedirectTable;
 
+	/**
+	 * Whether this camera director should run in the editor. 
+	 */
+	UPROPERTY(EditAnywhere, Category="Evaluation")
+	bool bRunInEditor = false;
+
 private:
 
 PRAGMA_DISABLE_DEPRECATION_WARNINGS
diff --git a/Engine/Plugins/Cameras/GameplayCameras/Source/GameplayCameras/Public/Nodes/Utility/BlueprintCameraNode.h b/Engine/Plugins/Cameras/GameplayCameras/Source/GameplayCameras/Public/Nodes/Utility/BlueprintCameraNode.h
--- a/Engine/Plugins/Cameras/GameplayCameras/Source/GameplayCameras/Public/Nodes/Utility/BlueprintCameraNode.h
+++ b/Engine/Plugins/Cameras/GameplayCameras/Source/GameplayCameras/Public/Nodes/Utility/BlueprintCameraNode.h
@@ -200,6 +200,10 @@ private:
 	UPROPERTY()
 	FCustomCameraNodeParameters CameraNodeEvaluatorOverrides;
 
+	/** Whether this camera node should run in the editor. */
+	UPROPERTY(EditAnywhere, Category=Common)
+	bool bRunInEditor = false;
+
 	// Deprecated.
 	
 	UPROPERTY()
