Cameras: the node parameter detail customization now supports several property metadata like min/max values or units

#Synced-CL 51024716

diff --git a/Engine/Plugins/Cameras/GameplayCameras/Source/GameplayCamerasEditor/Private/Customizations/CameraParameterDetailsCustomizations.cpp b/Engine/Plugins/Cameras/GameplayCameras/Source/GameplayCamerasEditor/Private/Customizations/CameraParameterDetailsCustomizations.cpp
--- a/Engine/Plugins/Cameras/GameplayCameras/Source/GameplayCamerasEditor/Private/Customizations/CameraParameterDetailsCustomizations.cpp
+++ b/Engine/Plugins/Cameras/GameplayCameras/Source/GameplayCamerasEditor/Private/Customizations/CameraParameterDetailsCustomizations.cpp
@@ -34,6 +34,16 @@
 namespace UE::Cameras
 {
 
+static TArrayView<FName> GetCameraParameterMetaDataToCopy()
+{
+	static TArray<FName> MetaDataKeys{ 
+		TEXT("UIMin"), TEXT("UIMax"),
+		TEXT("ClampMin"), TEXT("ClampMax"),
+		TEXT("Units"), TEXT("ForceUnits")
+	};
+	return MetaDataKeys;
+}
+
 void FCameraParameterDetailsCustomization::Register(FPropertyEditorModule& PropertyEditorModule)
 {
 #define UE_CAMERA_VARIABLE_FOR_TYPE(ValueType, ValueName)\
@@ -68,6 +78,16 @@ void FCameraParameterDetailsCustomization::CustomizeHeader(TSharedRef<IPropertyH
 	VariableProperty = PropertyHandle->GetChildHandle("Variable");
 	ensure(ValueProperty && VariableProperty);
 
+	// Copy select details-view-related metadata from the camera parameter struct onto the property handle.
+	for (const FName Key : GetCameraParameterMetaDataToCopy())
+	{
+		const FString& MetaDataValue = StructProperty->GetMetaData(Key);
+		if (!MetaDataValue.IsEmpty())
+		{
+			ValueProperty->SetInstanceMetaData(Key, MetaDataValue);
+		}
+	}
+
 	// Get the type of camera variable we need for this camera parameter (bool variable, float variable, etc.)
 	VariableClass = nullptr;
 	if (FObjectProperty* VariableObjectProperty = CastField<FObjectProperty>(VariableProperty->GetProperty()))
