Cameras: add support for first person view

- Add new properties to the camera pose for 1P views
- Add a new 1P node to set those properties

#Synced-CL 51024716

diff --git a/Engine/Plugins/Cameras/GameplayCameras/Source/GameplayCameras/Private/Core/CameraPose.cpp b/Engine/Plugins/Cameras/GameplayCameras/Source/GameplayCameras/Private/Core/CameraPose.cpp
--- a/Engine/Plugins/Cameras/GameplayCameras/Source/GameplayCameras/Private/Core/CameraPose.cpp
+++ b/Engine/Plugins/Cameras/GameplayCameras/Source/GameplayCameras/Private/Core/CameraPose.cpp
@@ -267,6 +267,18 @@ void FCameraPose::GetViewInfo(FMinimalViewInfo& OutViewInfo) const
 
 	OutViewInfo.PerspectiveNearClipPlane = NearClippingPlane > 0.0f ? NearClippingPlane : GNearClippingPlane;
 
+	OutViewInfo.bUseFirstPersonParameters = EnableFirstPerson;
+	if (EnableFirstPerson)
+	{
+		OutViewInfo.FirstPersonFOV = FirstPersonFieldOfView > 0 ? FirstPersonFieldOfView : OutViewInfo.FOV;
+		OutViewInfo.FirstPersonScale = FirstPersonScale > 0 ? FirstPersonScale : 1.f;
+	}
+	else
+	{
+		OutViewInfo.FirstPersonFOV = OutViewInfo.FOV;
+		OutViewInfo.FirstPersonScale = 1.f;
+	}
+
 	OutViewInfo.OffCenterProjectionOffset.X = GetHorizontalProjectionOffset();
 	OutViewInfo.OffCenterProjectionOffset.Y = GetVerticalProjectionOffset();
 
diff --git a/Engine/Plugins/Cameras/GameplayCameras/Source/GameplayCameras/Private/GameFramework/GameplayCameraComponentBase.cpp b/Engine/Plugins/Cameras/GameplayCameras/Source/GameplayCameras/Private/GameFramework/GameplayCameraComponentBase.cpp
--- a/Engine/Plugins/Cameras/GameplayCameras/Source/GameplayCameras/Private/GameFramework/GameplayCameraComponentBase.cpp
+++ b/Engine/Plugins/Cameras/GameplayCameras/Source/GameplayCameras/Private/GameFramework/GameplayCameraComponentBase.cpp
@@ -866,6 +866,26 @@ void UGameplayCameraComponentBase::UpdateOutputCameraComponent()
 			OutputCameraComponent->PostProcessSettings = Result.PostProcessSettings.Get();
 			OutputCameraComponent->PostProcessBlendWeight = 1.f;
 
+			if (Result.CameraPose.GetEnableFirstPerson())
+			{
+				const float FirstPersonFieldOfView = Result.CameraPose.GetFirstPersonFieldOfView();
+				const float FirstPersonScale = Result.CameraPose.GetFirstPersonScale();
+
+				OutputCameraComponent->bEnableFirstPersonFieldOfView = (FirstPersonFieldOfView > 0.f);
+				OutputCameraComponent->bEnableFirstPersonScale = (FirstPersonScale > 0.f);
+				OutputCameraComponent->FirstPersonFieldOfView = (
+						FirstPersonFieldOfView > 0.f ? FirstPersonFieldOfView : OutputCameraComponent->FieldOfView);
+				OutputCameraComponent->FirstPersonScale = (
+						FirstPersonScale > 0.f ? FirstPersonScale : 1.f);
+			}
+			else
+			{
+				OutputCameraComponent->bEnableFirstPersonFieldOfView = false;
+				OutputCameraComponent->bEnableFirstPersonScale = false;
+				OutputCameraComponent->FirstPersonFieldOfView = OutputCameraComponent->FieldOfView;
+				OutputCameraComponent->FirstPersonScale = 1.f;
+			}
+			
 			bGotValidTransform = true;
 		}
 	}
diff --git a/Engine/Plugins/Cameras/GameplayCameras/Source/GameplayCameras/Private/Nodes/Common/FirstPersonCameraNode.cpp b/Engine/Plugins/Cameras/GameplayCameras/Source/GameplayCameras/Private/Nodes/Common/FirstPersonCameraNode.cpp
new file mode 100644
--- /dev/null
+++ b/Engine/Plugins/Cameras/GameplayCameras/Source/GameplayCameras/Private/Nodes/Common/FirstPersonCameraNode.cpp
@@ -0,0 +1,69 @@
+// Copyright Epic Games, Inc. All Rights Reserved.
+
+#include "Nodes/Common/FirstPersonCameraNode.h"
+
+#include "Core/CameraParameterReader.h"
+
+#include UE_INLINE_GENERATED_CPP_BY_NAME(FirstPersonCameraNode)
+
+namespace UE::Cameras
+{
+
+class FFirstPersonCameraNodeEvaluator : public FCameraNodeEvaluator
+{
+	UE_DECLARE_CAMERA_NODE_EVALUATOR(GAMEPLAYCAMERAS_API, FFirstPersonCameraNodeEvaluator)
+
+protected:
+
+	// FCameraNodeEvaluator interface.
+	virtual void OnInitialize(const FCameraNodeEvaluatorInitializeParams& Params, FCameraNodeEvaluationResult& OutResult) override;
+	virtual void OnRun(const FCameraNodeEvaluationParams& Params, FCameraNodeEvaluationResult& OutResult) override;
+
+private:
+
+	TCameraParameterReader<bool> EnableFirstPersonReader;
+	TCameraParameterReader<float> FirstPersonFieldOfViewReader;
+	TCameraParameterReader<float> FirstPersonScaleReader;
+
+};
+
+UE_DEFINE_CAMERA_NODE_EVALUATOR(FFirstPersonCameraNodeEvaluator)
+
+void FFirstPersonCameraNodeEvaluator::OnInitialize(const FCameraNodeEvaluatorInitializeParams& Params, FCameraNodeEvaluationResult& OutResult)
+{
+	SetNodeEvaluatorFlags(ECameraNodeEvaluatorFlags::None);
+
+	const UFirstPersonCameraNode* FirstPersonCameraNode = GetCameraNodeAs<UFirstPersonCameraNode>();
+	EnableFirstPersonReader.Initialize(FirstPersonCameraNode->EnableFirstPerson);
+	FirstPersonFieldOfViewReader.Initialize(FirstPersonCameraNode->FirstPersonFieldOfView);
+	FirstPersonScaleReader.Initialize(FirstPersonCameraNode->FirstPersonScale);
+}
+
+void FFirstPersonCameraNodeEvaluator::OnRun(const FCameraNodeEvaluationParams& Params, FCameraNodeEvaluationResult& OutResult)
+{
+	const UFirstPersonCameraNode* FirstPersonCameraNode = GetCameraNodeAs<UFirstPersonCameraNode>();
+
+	const bool bEnableFirstPerson = EnableFirstPersonReader.Get(OutResult.VariableTable);
+	OutResult.CameraPose.SetEnableFirstPerson(bEnableFirstPerson);
+
+	const float FirstPersonFieldOfView = FirstPersonFieldOfViewReader.Get(OutResult.VariableTable);
+	if (FirstPersonCameraNode->bEnableFirstPersonFieldOfView && FirstPersonFieldOfView > 0.f)
+	{
+		OutResult.CameraPose.SetFirstPersonFieldOfView(FirstPersonFieldOfView);
+	}
+
+	const float FirstPersonScale = FirstPersonScaleReader.Get(OutResult.VariableTable);
+	if (FirstPersonCameraNode->bEnableFirstPersonScale && FirstPersonScale > 0.f)
+	{
+		OutResult.CameraPose.SetFirstPersonScale(FirstPersonScale);
+	}
+}
+
+}  // namespace UE::Cameras
+
+FCameraNodeEvaluatorPtr UFirstPersonCameraNode::OnBuildEvaluator(FCameraNodeEvaluatorBuilder& Builder) const
+{
+	using namespace UE::Cameras;
+	return Builder.BuildEvaluator<FFirstPersonCameraNodeEvaluator>();
+}
+
diff --git a/Engine/Plugins/Cameras/GameplayCameras/Source/GameplayCameras/Public/Core/CameraPose.h b/Engine/Plugins/Cameras/GameplayCameras/Source/GameplayCameras/Public/Core/CameraPose.h
--- a/Engine/Plugins/Cameras/GameplayCameras/Source/GameplayCameras/Public/Core/CameraPose.h
+++ b/Engine/Plugins/Cameras/GameplayCameras/Source/GameplayCameras/Public/Core/CameraPose.h
@@ -37,13 +37,16 @@ enum EAspectRatioAxisConstraint : int;
 	UE_CAMERA_POSE_FOR_PROPERTY(float,  PhysicalCameraBlendWeight)\
 	UE_CAMERA_POSE_FOR_PROPERTY(float,  NearClippingPlane)\
 	UE_CAMERA_POSE_FOR_PROPERTY(float,  FarClippingPlane)\
-	UE_CAMERA_POSE_FOR_PROPERTY(float,  OrthographicWidth)
+	UE_CAMERA_POSE_FOR_PROPERTY(float,  OrthographicWidth)\
+	UE_CAMERA_POSE_FOR_PROPERTY(float,  FirstPersonFieldOfView)\
+	UE_CAMERA_POSE_FOR_PROPERTY(float,  FirstPersonScale)
 
 #define UE_CAMERA_POSE_FOR_FOV_PROPERTIES()\
 	UE_CAMERA_POSE_FOR_PROPERTY(float, FieldOfView)\
 	UE_CAMERA_POSE_FOR_PROPERTY(float, FocalLength)
 
 #define UE_CAMERA_POSE_FOR_FLIPPING_PROPERTIES()\
+	UE_CAMERA_POSE_FOR_PROPERTY(bool, EnableFirstPerson)\
 	UE_CAMERA_POSE_FOR_PROPERTY(bool, EnablePhysicalCamera)\
 	UE_CAMERA_POSE_FOR_PROPERTY(bool, ConstrainAspectRatio)\
 	UE_CAMERA_POSE_FOR_PROPERTY(bool, OverrideAspectRatioAxisConstraint)\
@@ -274,9 +277,15 @@ private:
 	UPROPERTY()
 	float FocalLength = 35.f;
 
-	/**
-	 * The desired width (in world units) of the orthographic view (ignored in Perspective mode) 
-	 */
+	/** The horizontal field of view (in degrees) used for primitives tagged as "IsFirstPerson" */
+	UPROPERTY()
+	float FirstPersonFieldOfView = -1.f;
+
+	/** The scale to apply to primitives tagged as "IsFirstPerson" */
+	UPROPERTY()
+	float FirstPersonScale = 1.f;
+	
+	/** The desired width (in world units) of the orthographic view (ignored in Perspective mode) */
 	UPROPERTY()
 	float OrthographicWidth = DEFAULT_ORTHOWIDTH;
 
@@ -339,6 +348,10 @@ private:
 	UPROPERTY()
 	float PhysicalCameraBlendWeight = 0.f;
 
+	/** Whether to enable first-person parameters (FOV, scale) for primitives tagged as "IsFirstPerson" */
+	UPROPERTY()
+	bool EnableFirstPerson = false;
+
 	/** 
 	 * Whether to setup post-process settings based on physical camera properties such as Aperture,
 	 * FocusDistance, DiaphragmBladeCount, and so on.
diff --git a/Engine/Plugins/Cameras/GameplayCameras/Source/GameplayCameras/Public/GameFramework/BlueprintCameraPose.h b/Engine/Plugins/Cameras/GameplayCameras/Source/GameplayCameras/Public/GameFramework/BlueprintCameraPose.h
--- a/Engine/Plugins/Cameras/GameplayCameras/Source/GameplayCameras/Public/GameFramework/BlueprintCameraPose.h
+++ b/Engine/Plugins/Cameras/GameplayCameras/Source/GameplayCameras/Public/GameFramework/BlueprintCameraPose.h
@@ -50,6 +50,14 @@ public:
 	UPROPERTY(BlueprintReadWrite, EditAnywhere, Category="Camera")
 	float FocalLength = 35.f;
 
+	/** The horizontal field of view (in degrees) used for primitives tagged as "IsFirstPerson" */
+	UPROPERTY(BlueprintReadWrite, EditAnywhere, Category="Camera")
+	float FirstPersonFieldOfView = 90.f;
+
+	/** The scale to apply to primitives tagged as "IsFirstPerson" */
+	UPROPERTY(BlueprintReadWrite, EditAnywhere, Category="Camera")
+	float FirstPersonScale = 1.f;
+
 	/** The desired width (in world units) of the orthographic view (ignored in Perspective mode) */
 	UPROPERTY(BlueprintReadWrite, EditAnywhere, Category = "Camera")
 	float OrthographicWidth = DEFAULT_ORTHOWIDTH;
@@ -109,6 +117,10 @@ public:
 	/** Internal weight for physical camera post-process settings. */
 	UPROPERTY()
 	float PhysicalCameraBlendWeight = 0.f;
+	
+	/** Whether to enable first-person parameters (FOV, scale) for primitives tagged as "IsFirstPerson" */
+	UPROPERTY(EditAnywhere, BlueprintReadWrite, Category="Camera")
+	bool EnableFirstPerson = false;
 
 	/** 
 	 * Whether to setup post-process settings based on physical camera properties such as Aperture,
diff --git a/Engine/Plugins/Cameras/GameplayCameras/Source/GameplayCameras/Public/Nodes/Common/FirstPersonCameraNode.h b/Engine/Plugins/Cameras/GameplayCameras/Source/GameplayCameras/Public/Nodes/Common/FirstPersonCameraNode.h
new file mode 100644
--- /dev/null
+++ b/Engine/Plugins/Cameras/GameplayCameras/Source/GameplayCameras/Public/Nodes/Common/FirstPersonCameraNode.h
@@ -0,0 +1,56 @@
+// Copyright Epic Games, Inc. All Rights Reserved.
+
+#pragma once
+
+#include "Core/CameraNode.h"
+#include "Core/CameraParameters.h"
+
+#include "FirstPersonCameraNode.generated.h"
+
+/**
+ * A camera node that sets and enables the first-person parameters on the resulting camera view.
+ */
+UCLASS(MinimalAPI, meta=(CameraNodeCategories="Common"))
+class UFirstPersonCameraNode : public UCameraNode
+{
+	GENERATED_BODY()
+
+protected:
+
+	// UCameraNode interface.
+	virtual FCameraNodeEvaluatorPtr OnBuildEvaluator(FCameraNodeEvaluatorBuilder& Builder) const override;
+
+public:
+	
+	/** Whether to enable first-person view parameters on primitives tagged as "IsFirstPerson". */
+	UPROPERTY(EditAnywhere, Category=Common)
+	FBooleanCameraParameter EnableFirstPerson = true;
+
+	/** 
+	 * The horizontal field of view (in degrees) used for primitives tagged as "IsFirstPerson".
+	 * If zero or negative, use the normal field of view.
+	 */
+	UPROPERTY(EditAnywhere, Category=Common, meta=(
+				UIMin="5.0", UIMax="170", ClampMin="0.001", ClampMax="360.0", Units=deg, 
+				EditCondition="bEnableFirstPersonFieldOfView"))
+	FFloatCameraParameter FirstPersonFieldOfView = 90.f;
+
+	/**
+	 * The scale to apply to primitives tagged as "IsFirstPerson". 
+	 *
+	 * This is used to scale down these primitives and prevent them from intersecting with the scene.
+	 */
+	UPROPERTY(EditAnywhere, Category=Common, meta=(
+				UIMin="0.001", UIMax="1.0", ClampMin="0.001", Units=%,
+				EditCondition="bEnableFirstPersonScale"))
+	FFloatCameraParameter FirstPersonScale = 1.f;
+
+public:
+
+	UPROPERTY(EditAnywhere, Category=Common, meta=(InlineEditConditionToggle))
+	bool bEnableFirstPersonFieldOfView = true;
+
+	UPROPERTY(EditAnywhere, Category=Common, meta=(InlineEditConditionToggle))
+	bool bEnableFirstPersonScale = true;
+};
+
