Cameras: add explicit checkboxes to the lens parameters node

Instead of specifying zero or negative values to indicate that some property should be ignored, there are now explicit checkboxes to enable/disable each lens parameter.

#Synced-CL 51024716

diff --git a/Engine/Plugins/Cameras/GameplayCameras/Source/GameplayCameras/Private/GameplayCamerasCustomVersion.cpp b/Engine/Plugins/Cameras/GameplayCameras/Source/GameplayCameras/Private/GameplayCamerasCustomVersion.cpp
new file mode 100644
--- /dev/null
+++ b/Engine/Plugins/Cameras/GameplayCameras/Source/GameplayCameras/Private/GameplayCamerasCustomVersion.cpp
@@ -0,0 +1,11 @@
+// Copyright Epic Games, Inc. All Rights Reserved.
+
+#include "GameplayCamerasCustomVersion.h"
+
+#include "Serialization/CustomVersion.h"
+
+FCustomVersionRegistration GGameplayCamerasCustomVersion(
+		FGameplayCamerasCustomVersion::GUID,
+		FGameplayCamerasCustomVersion::LatestVersion,
+		TEXT("GameplayCamerasCustomVersion"));
+
diff --git a/Engine/Plugins/Cameras/GameplayCameras/Source/GameplayCameras/Private/Nodes/Common/LensParametersCameraNode.cpp b/Engine/Plugins/Cameras/GameplayCameras/Source/GameplayCameras/Private/Nodes/Common/LensParametersCameraNode.cpp
--- a/Engine/Plugins/Cameras/GameplayCameras/Source/GameplayCameras/Private/Nodes/Common/LensParametersCameraNode.cpp
+++ b/Engine/Plugins/Cameras/GameplayCameras/Source/GameplayCameras/Private/Nodes/Common/LensParametersCameraNode.cpp
@@ -2,9 +2,8 @@
 
 #include "Nodes/Common/LensParametersCameraNode.h"
 
-#include "Core/CameraEvaluationContext.h"
 #include "Core/CameraParameterReader.h"
-#include "GameplayCameras.h"
+#include "GameplayCamerasCustomVersion.h"
 
 #include UE_INLINE_GENERATED_CPP_BY_NAME(LensParametersCameraNode)
 
@@ -24,8 +23,8 @@ protected:
 private:
 
 	TCameraParameterReader<float> FocalLengthReader;
-	TCameraParameterReader<float> FocusDistanceReader;
 	TCameraParameterReader<float> ApertureReader;
+	TCameraParameterReader<float> FocusDistanceReader;
 	TCameraParameterReader<bool> EnablePhysicalCameraReader;
 };
 
@@ -37,31 +36,32 @@ void FLensParametersCameraNodeEvaluator::OnInitialize(const FCameraNodeEvaluator
 
 	const ULensParametersCameraNode* LensParametersNode = GetCameraNodeAs<ULensParametersCameraNode>();
 	FocalLengthReader.Initialize(LensParametersNode->FocalLength);
-	FocusDistanceReader.Initialize(LensParametersNode->FocusDistance);
 	ApertureReader.Initialize(LensParametersNode->Aperture);
+	FocusDistanceReader.Initialize(LensParametersNode->FocusDistance);
 	EnablePhysicalCameraReader.Initialize(LensParametersNode->EnablePhysicalCamera);
 }
 
 void FLensParametersCameraNodeEvaluator::OnRun(const FCameraNodeEvaluationParams& Params, FCameraNodeEvaluationResult& OutResult)
 {
+	const ULensParametersCameraNode* LensParametersNode = GetCameraNodeAs<ULensParametersCameraNode>();
 	FCameraPose& OutPose = OutResult.CameraPose;
 
 	float FocalLength = FocalLengthReader.Get(OutResult.VariableTable);
-	if (FocalLength > 0)
+	if (LensParametersNode->bEnableFocalLength && FocalLength > 0)
 	{
 		OutPose.SetFocalLength(FocalLength);
 		OutPose.SetFieldOfView(-1);
 	}
-	float FocusDistance = FocusDistanceReader.Get(OutResult.VariableTable);
-	if (FocusDistance > 0)
-	{
-		OutPose.SetFocusDistance(FocusDistance);
-	}
 	float Aperture = ApertureReader.Get(OutResult.VariableTable);
-	if (Aperture > 0)
+	if (LensParametersNode->bEnableAperture && Aperture > 0)
 	{
 		OutPose.SetAperture(Aperture);
 	}
+	float FocusDistance = FocusDistanceReader.Get(OutResult.VariableTable);
+	if (LensParametersNode->bEnableFocusDistance && FocusDistance > 0)
+	{
+		OutPose.SetFocusDistance(FocusDistance);
+	}
 
 	const bool bEnablePhysicalCamera = EnablePhysicalCameraReader.Get(OutResult.VariableTable);
 	OutPose.SetEnablePhysicalCamera(bEnablePhysicalCamera);
@@ -70,6 +70,51 @@ void FLensParametersCameraNodeEvaluator::OnRun(const FCameraNodeEvaluationParams
 
 }  // namespace UE::Cameras
 
+void ULensParametersCameraNode::Serialize(FArchive& Ar)
+{
+	Ar.UsingCustomVersion(FGameplayCamerasCustomVersion::GUID);
+
+	// If this node was saved before the parameter flags were added, set the properties to their old default values,
+	// then do the serialization, which will give us the non-zero values.
+	const bool bUpgradeParameterFlags = (
+			Ar.IsLoading() &&
+			Ar.CustomVer(FGameplayCamerasCustomVersion::GUID) < FGameplayCamerasCustomVersion::AddLensNodeParameterFlags);
+	if (bUpgradeParameterFlags)
+	{
+		FocalLength.Value = 0;
+		Aperture.Value = 0;
+		FocusDistance.Value = 0;
+
+		bEnableFocalLength = false;
+		bEnableAperture = false;
+		bEnableFocusDistance = false;
+	}
+
+	Super::Serialize(Ar);
+
+	// Now that we have the non-zero values, auto-set the appropriate flags, and restore the new defaults.
+	// These defaults should match the ones in the header file!
+	if (bUpgradeParameterFlags)
+	{
+		bEnableFocalLength = FocalLength.Value > 0.f || FocalLength.Variable || FocalLength.VariableID;
+		bEnableAperture = Aperture.Value > 0.f || Aperture.Variable || Aperture.VariableID;
+		bEnableFocusDistance = FocusDistance.Value > 0.f || FocusDistance.Variable || FocusDistance.VariableID;
+
+		if (FocalLength.Value <= 0.f)
+		{
+			FocalLength.Value = 35.f;
+		}
+		if (Aperture.Value <= 0.f)
+		{
+			Aperture.Value = 16.f;
+		}
+		if (FocusDistance.Value <= 0.f)
+		{
+			FocusDistance.Value = 1000.f;
+		}
+	}
+}
+
 FCameraNodeEvaluatorPtr ULensParametersCameraNode::OnBuildEvaluator(FCameraNodeEvaluatorBuilder& Builder) const
 {
 	using namespace UE::Cameras;
diff --git a/Engine/Plugins/Cameras/GameplayCameras/Source/GameplayCameras/Public/GameplayCamerasCustomVersion.h b/Engine/Plugins/Cameras/GameplayCameras/Source/GameplayCameras/Public/GameplayCamerasCustomVersion.h
new file mode 100644
--- /dev/null
+++ b/Engine/Plugins/Cameras/GameplayCameras/Source/GameplayCameras/Public/GameplayCamerasCustomVersion.h
@@ -0,0 +1,20 @@
+// Copyright Epic Games, Inc. All Rights Reserved.
+
+#pragma once
+
+#include "Misc/Guid.h"
+
+namespace FGameplayCamerasCustomVersion
+{
+	enum Type
+	{
+		BeforeCustomVersionWasAdded = 0,
+		AddLensNodeParameterFlags = 1,
+
+		VersionPlusOne,
+		LatestVersion = VersionPlusOne - 1
+	};
+
+	const FGuid GUID(0x8BCB8CA9, 0x642B497C, 0xAB759854, 0x44875CE1);
+}
+
diff --git a/Engine/Plugins/Cameras/GameplayCameras/Source/GameplayCameras/Public/Nodes/Common/LensParametersCameraNode.h b/Engine/Plugins/Cameras/GameplayCameras/Source/GameplayCameras/Public/Nodes/Common/LensParametersCameraNode.h
--- a/Engine/Plugins/Cameras/GameplayCameras/Source/GameplayCameras/Public/Nodes/Common/LensParametersCameraNode.h
+++ b/Engine/Plugins/Cameras/GameplayCameras/Source/GameplayCameras/Public/Nodes/Common/LensParametersCameraNode.h
@@ -16,6 +16,11 @@ class ULensParametersCameraNode : public UCameraNode
 {
 	GENERATED_BODY()
 
+public:
+
+	// UObject interface.
+	virtual void Serialize(FArchive& Ar) override;
+
 protected:
 
 	// UCameraNode interface.
@@ -24,16 +29,22 @@ protected:
 public:
 
 	/** Current focal length of the camera (i.e. controls FoV, zoom) */
-	UPROPERTY(EditAnywhere, Category="Lens Parameters")
-	FFloatCameraParameter FocalLength;
-
-	/** Manually-controlled focus distance (manual focus mode only) */
-	UPROPERTY(EditAnywhere, Category="Lens Parameters", meta=(Units=cm))
-	FFloatCameraParameter FocusDistance;
+	UPROPERTY(EditAnywhere, Category="Lens Parameters", meta=(
+				UIMin=5.0, UIMax=200.0, ClampMin=0.001, ForceUnits=mm,
+				EditCondition="bEnableFocalLength"))
+	FFloatCameraParameter FocalLength = { 35.f };
 
 	/** Current aperture, in terms of f-stop (e.g. 2.8 for f/2.8) */
-	UPROPERTY(EditAnywhere, Category="Lens Parameters")
-	FFloatCameraParameter Aperture;
+	UPROPERTY(EditAnywhere, Category="Lens Parameters", meta=(
+				UIMin=0.1, UIMax=20.0, ClampMin=0.001,
+				EditCondition="bEnableAperture"))
+	FFloatCameraParameter Aperture = { 16.f };
+
+	/** Manually-controlled focus distance (manual focus mode only) */
+	UPROPERTY(EditAnywhere, Category="Lens Parameters", meta=(
+				UIMin=0.001, UIMax=100000.0, ClampMin=0.001,
+				EditCondition="bEnableFocusDistance"))
+	FFloatCameraParameter FocusDistance = { 1000.f };
 
 	/** 
 	 * Whether post-process settings should be automatically applied to reflect
@@ -41,5 +52,19 @@ public:
 	 */
 	UPROPERTY(EditAnywhere, Category="Lens Parameters")
 	FBooleanCameraParameter EnablePhysicalCamera = false;
+
+public:
+
+	/** Enable setting the focal length. When false, focal length is not set. */
+	UPROPERTY(EditAnywhere, Category="Lens Features", meta=(InlineEditConditionToggle))
+	bool bEnableFocalLength = true;
+
+	/** Enable setting the aperture. When false, aperture is not set. */
+	UPROPERTY(EditAnywhere, Category="Lens Features", meta=(InlineEditConditionToggle))
+	bool bEnableAperture = true;
+
+	/** Enable setting the focus distance. When false, focus distance is not set. */
+	UPROPERTY(EditAnywhere, Category="Lens Features", meta=(InlineEditConditionToggle))
+	bool bEnableFocusDistance = false;
 };
 
