Cameras: expand documentation and add UI min/max settings to the field of view node

#Synced-CL 51024716

diff --git a/Engine/Plugins/Cameras/GameplayCameras/Source/GameplayCameras/Public/Nodes/Common/FieldOfViewCameraNode.h b/Engine/Plugins/Cameras/GameplayCameras/Source/GameplayCameras/Public/Nodes/Common/FieldOfViewCameraNode.h
--- a/Engine/Plugins/Cameras/GameplayCameras/Source/GameplayCameras/Public/Nodes/Common/FieldOfViewCameraNode.h
+++ b/Engine/Plugins/Cameras/GameplayCameras/Source/GameplayCameras/Public/Nodes/Common/FieldOfViewCameraNode.h
@@ -18,8 +18,13 @@ class UFieldOfViewCameraNode : public UCameraNode
 
 public:
 
-	/** The field of view, in degrees. */
-	UPROPERTY(EditAnywhere, Category=Common)
+	/**
+	 * The horizontal field of view, in degrees.
+	 *
+	 * Note that the effective horizontal field of view may be different, depending on the viewport size and the way
+	 * aspect ratio constraints have been setup on the camera.
+	 */
+	UPROPERTY(EditAnywhere, Category=Common, meta=(UIMin="5.0", UIMax="170", ClampMin="0.001", ClampMax="360.0", Units=deg))
 	FFloatCameraParameter FieldOfView;
 
 public:
