Cameras: allow using an "output" camera component created in the parent actor

Actor components cannot create visible/editor-friendly children components themselves. Component hierarchies need to be created by the actor. This change moves the creation of the "output" camera component to the parent actor. The GPC component only creates it if they can't find one created by the actor, which can happen if the GPC is inside another actor who doesn't know about all this, such as a character pawn or any other Blueprint. The user can still get this output camera component to be visible by creating it themselves and parenting it under the GPC component.

Since this effectily changes the output camera component from a "hidden" implementation detail to an "official" component of the actor, it would now show up in the Details View of the actor by default, cluttering up the UI with irrelevant properties that get overriden by the camera system evaluation. So this change also tweaks the categories and metadata on a few properties to make sure this component is hidden from the actor's Details View.

#Synced-CL 51024716

diff --git a/Engine/Plugins/Cameras/GameplayCameras/Source/GameplayCameras/Private/GameFramework/GameplayCameraActor.cpp b/Engine/Plugins/Cameras/GameplayCameras/Source/GameplayCameras/Private/GameFramework/GameplayCameraActor.cpp
--- a/Engine/Plugins/Cameras/GameplayCameras/Source/GameplayCameras/Private/GameFramework/GameplayCameraActor.cpp
+++ b/Engine/Plugins/Cameras/GameplayCameras/Source/GameplayCameras/Private/GameFramework/GameplayCameraActor.cpp
@@ -2,6 +2,7 @@
 
 #include "GameFramework/GameplayCameraActor.h"
 
+#include "CineCameraComponent.h"
 #include "Engine/EngineTypes.h"
 #include "Engine/World.h"
 #include "GameFramework/GameplayCameraComponent.h"
@@ -16,6 +17,9 @@ AGameplayCameraActor::AGameplayCameraActor(const FObjectInitializer& ObjectInit)
 {
 	CameraComponent = CreateDefaultSubobject<UGameplayCameraComponent>(TEXT("CameraComponent"));
 	RootComponent = CameraComponent;
+
+	OutputCameraComponent = CreateDefaultSubobject<UCineCameraComponent>(TEXT("OutputCameraComponent"), true);
+	OutputCameraComponent->SetupAttachment(CameraComponent);
 }
 
 USceneComponent* AGameplayCameraActor::GetDefaultAttachComponent() const
diff --git a/Engine/Plugins/Cameras/GameplayCameras/Source/GameplayCameras/Private/GameFramework/GameplayCameraComponentBase.cpp b/Engine/Plugins/Cameras/GameplayCameras/Source/GameplayCameras/Private/GameFramework/GameplayCameraComponentBase.cpp
--- a/Engine/Plugins/Cameras/GameplayCameras/Source/GameplayCameras/Private/GameFramework/GameplayCameraComponentBase.cpp
+++ b/Engine/Plugins/Cameras/GameplayCameras/Source/GameplayCameras/Private/GameFramework/GameplayCameraComponentBase.cpp
@@ -203,11 +203,6 @@ bool UGameplayCameraComponentBase::EnsureCameraSystemHostIfNeeded()
 				*GetNameSafe(GetOwner()), *GetNameSafe(this));
 	}
 
-	if (ensure(OutputCameraComponent))
-	{
-		OutputCameraComponent->Activate();
-	}
-
 	return true;
 }
 
@@ -217,7 +212,6 @@ void UGameplayCameraComponentBase::DestroyCameraSystemHost()
 	if (OutputCameraComponent)
 	{
 		OutputCameraComponent->SetRelativeTransform(FTransform());
-		OutputCameraComponent->Deactivate();
 	}
 
 	DestroyCameraSystem();
@@ -511,15 +505,6 @@ void UGameplayCameraComponentBase::OnRegister()
 
 	Super::OnRegister();
 
-	if (OutputCameraComponent == nullptr)
-	{
-		OutputCameraComponent = NewObject<UCineCameraComponent>(this, TEXT("OutputCameraComponent"), RF_Transient | RF_TextExportTransient);
-		OutputCameraComponent->SetAutoActivate(false);
-		OutputCameraComponent->SetupAttachment(this);
-		OutputCameraComponent->CreationMethod = CreationMethod;
-		OutputCameraComponent->RegisterComponentWithWorld(GetWorld());
-	}
-
 #if WITH_EDITOR
 
 	UWorld* World = GetWorld();
@@ -567,7 +552,7 @@ void UGameplayCameraComponentBase::EndPlay(const EEndPlayReason::Type EndPlayRea
 
 void UGameplayCameraComponentBase::OnComponentDestroyed(bool bDestroyingHierarchy)
 {
-	if (OutputCameraComponent)
+	if (OutputCameraComponent && bIsOwnedOutputCameraComponent)
 	{
 		OutputCameraComponent->DestroyComponent();
 	}
@@ -587,6 +572,11 @@ void UGameplayCameraComponentBase::Activate(bool bReset)
 	{
 		EnsureEvaluationContext(nullptr);
 		EnsureCameraSystemHostIfNeeded();
+
+		if (OutputCameraComponent && bIsOwnedOutputCameraComponent)
+		{
+			OutputCameraComponent->Activate(bReset);
+		}
 	}
 }
 
@@ -595,6 +585,11 @@ void UGameplayCameraComponentBase::Deactivate()
 	DestroyCameraSystemHost();
 	DestroyEvaluationContext();
 
+	if (OutputCameraComponent && bIsOwnedOutputCameraComponent)
+	{
+		OutputCameraComponent->Deactivate();
+	}
+
 	Super::Deactivate();
 }
 
@@ -819,10 +814,57 @@ void UGameplayCameraComponentBase::RecreateEditorWorldEvaluationContext()
 
 #endif  // WITH_EDITOR
 
+void UGameplayCameraComponentBase::EnsureOutputCameraComponent()
+{
+	if (OutputCameraComponent)
+	{
+		return;
+	}
+
+	// If our parent actor has put a camera component under us, we can use it. This is useful for making that camera
+	// component "visible" to the user in the editor.
+	for (USceneComponent* AttachChild : GetAttachChildren())
+	{
+		if (UCineCameraComponent* ChildCamera = Cast<UCineCameraComponent>(AttachChild))
+		{
+			OutputCameraComponent = ChildCamera;
+			bIsOwnedOutputCameraComponent = false;
+			break;
+		}
+	}
+
+	// If we didn't find any child camera component to use, we need to create one ourselves. A component cannot create 
+	// an editor-friendly child component that is visible, so we leave that one hidden (i.e. held in a non-visible 
+	// property)
+	if (OutputCameraComponent == nullptr)
+	{
+		OutputCameraComponent = NewObject<UCineCameraComponent>(this, TEXT("OutputCameraComponent"), RF_Transient);
+		OutputCameraComponent->SetAutoActivate(false);
+		OutputCameraComponent->SetupAttachment(this);
+		OutputCameraComponent->CreationMethod = CreationMethod;
+		OutputCameraComponent->RegisterComponentWithWorld(GetWorld());
+
+		bIsOwnedOutputCameraComponent = true;
+	}
+
+	// Make sure the component is active, otherwise it won't be picked up as a valid view by the engine.
+	if (ensure(OutputCameraComponent))
+	{
+		OutputCameraComponent->Activate();
+	}
+}
+
 void UGameplayCameraComponentBase::UpdateOutputCameraComponent()
 {
 	using namespace UE::Cameras;
 
+	// Make sure we have a camera component to write to.
+	//
+	// Ideally this would be done once in OnRegister() but we don't know at that time about our attached children.
+	// The next best place would be in BeginPlay(), but that's not called in editor contexts.
+	// So we lazily do it here instead.
+	EnsureOutputCameraComponent();
+
 	if (!OutputCameraComponent)
 	{
 		return;
diff --git a/Engine/Plugins/Cameras/GameplayCameras/Source/GameplayCameras/Private/GameFramework/GameplayCameraRigActor.cpp b/Engine/Plugins/Cameras/GameplayCameras/Source/GameplayCameras/Private/GameFramework/GameplayCameraRigActor.cpp
--- a/Engine/Plugins/Cameras/GameplayCameras/Source/GameplayCameras/Private/GameFramework/GameplayCameraRigActor.cpp
+++ b/Engine/Plugins/Cameras/GameplayCameras/Source/GameplayCameras/Private/GameFramework/GameplayCameraRigActor.cpp
@@ -2,6 +2,7 @@
 
 #include "GameFramework/GameplayCameraRigActor.h"
 
+#include "CineCameraComponent.h"
 #include "Engine/EngineTypes.h"
 #include "Engine/World.h"
 #include "GameFramework/GameplayCameraRigComponent.h"
@@ -16,6 +17,9 @@ AGameplayCameraRigActor::AGameplayCameraRigActor(const FObjectInitializer& Objec
 {
 	CameraRigComponent = CreateDefaultSubobject<UGameplayCameraRigComponent>(TEXT("CameraRigComponent"));
 	RootComponent = CameraRigComponent;
+
+	OutputCameraComponent = CreateDefaultSubobject<UCineCameraComponent>(TEXT("OutputCameraComponent"), true);
+	OutputCameraComponent->SetupAttachment(CameraRigComponent);
 }
 
 USceneComponent* AGameplayCameraRigActor::GetDefaultAttachComponent() const
diff --git a/Engine/Plugins/Cameras/GameplayCameras/Source/GameplayCameras/Public/GameFramework/GameplayCameraActor.h b/Engine/Plugins/Cameras/GameplayCameras/Source/GameplayCameras/Public/GameFramework/GameplayCameraActor.h
--- a/Engine/Plugins/Cameras/GameplayCameras/Source/GameplayCameras/Public/GameFramework/GameplayCameraActor.h
+++ b/Engine/Plugins/Cameras/GameplayCameras/Source/GameplayCameras/Public/GameFramework/GameplayCameraActor.h
@@ -6,12 +6,13 @@
 
 #include "GameplayCameraActor.generated.h"
 
+class UCineCameraComponent;
 class UGameplayCameraComponent;
 
 /**
  * An actor that can run a camera asset.
  */
-UCLASS(BlueprintType, MinimalAPI, ClassGroup=Camera, HideCategories=(Input, Rendering))
+UCLASS(BlueprintType, MinimalAPI, ClassGroup=Camera, HideCategories=(OutputCamera))
 class AGameplayCameraActor : public AGameplayCameraActorBase
 {
 	GENERATED_BODY()
@@ -26,6 +27,10 @@ public:
 	UFUNCTION(BlueprintGetter, Category=Camera)
 	UGameplayCameraComponent* GetCameraComponent() const { return CameraComponent; }
 
+	/** Gets the output camera component. */
+	UFUNCTION(BlueprintGetter, Category=Camera)
+	UCineCameraComponent* GetOutputCameraComponent() const { return OutputCameraComponent; }
+
 public:
 
 	// AActor interface.
@@ -38,7 +43,10 @@ protected:
 
 private:
 
-	UPROPERTY(VisibleAnywhere, Category=Camera, BlueprintGetter="GetCameraComponent", meta=(ExposeFunctionCategories="Camera"))
+	UPROPERTY(VisibleAnywhere, Category=GameplayCamera, BlueprintGetter="GetCameraComponent", meta=(ExposeFunctionCategories="Camera"))
 	TObjectPtr<UGameplayCameraComponent> CameraComponent;
+
+	UPROPERTY(Transient, VisibleAnywhere, Category=OutputCamera, BlueprintGetter="GetOutputCameraComponent")
+	TObjectPtr<UCineCameraComponent> OutputCameraComponent;
 };
 
diff --git a/Engine/Plugins/Cameras/GameplayCameras/Source/GameplayCameras/Public/GameFramework/GameplayCameraComponentBase.h b/Engine/Plugins/Cameras/GameplayCameras/Source/GameplayCameras/Public/GameFramework/GameplayCameraComponentBase.h
--- a/Engine/Plugins/Cameras/GameplayCameras/Source/GameplayCameras/Public/GameFramework/GameplayCameraComponentBase.h
+++ b/Engine/Plugins/Cameras/GameplayCameras/Source/GameplayCameras/Public/GameFramework/GameplayCameraComponentBase.h
@@ -242,6 +242,7 @@ private:
 	void EnsureEvaluationContext(APlayerController* PlayerController);
 	void DestroyEvaluationContext();
 
+	void EnsureOutputCameraComponent();
 	void UpdateOutputCameraComponent();
 
 #if WITH_EDITOR
@@ -297,6 +298,10 @@ private:
 	UPROPERTY(Transient)
 	TObjectPtr<UCineCameraComponent> OutputCameraComponent;
 
+	/** Whether the output camera component was auto-created, or obtained from the parent actor. */
+	UPROPERTY()
+	bool bIsOwnedOutputCameraComponent = false;
+
 private:
 
 	using FGameplayCameraComponentEvaluationContext = UE::Cameras::FGameplayCameraComponentEvaluationContext;
diff --git a/Engine/Plugins/Cameras/GameplayCameras/Source/GameplayCameras/Public/GameFramework/GameplayCameraRigActor.h b/Engine/Plugins/Cameras/GameplayCameras/Source/GameplayCameras/Public/GameFramework/GameplayCameraRigActor.h
--- a/Engine/Plugins/Cameras/GameplayCameras/Source/GameplayCameras/Public/GameFramework/GameplayCameraRigActor.h
+++ b/Engine/Plugins/Cameras/GameplayCameras/Source/GameplayCameras/Public/GameFramework/GameplayCameraRigActor.h
@@ -6,12 +6,13 @@
 
 #include "GameplayCameraRigActor.generated.h"
 
+class UCineCameraComponent;
 class UGameplayCameraRigComponent;
 
 /**
  * An actor that can run a camera asset.
  */
-UCLASS(BlueprintType, MinimalAPI, ClassGroup=Camera, HideCategories=(Input, Rendering))
+UCLASS(BlueprintType, MinimalAPI, ClassGroup=Camera, HideCategories=(OutputCamera))
 class AGameplayCameraRigActor : public AGameplayCameraActorBase
 {
 	GENERATED_BODY()
@@ -26,6 +27,10 @@ public:
 	UFUNCTION(BlueprintGetter, Category=Camera)
 	UGameplayCameraRigComponent* GetCameraRigComponent() const { return CameraRigComponent; }
 
+	/** Gets the output camera component. */
+	UFUNCTION(BlueprintGetter, Category=Camera)
+	UCineCameraComponent* GetOutputCameraComponent() const { return OutputCameraComponent; }
+
 public:
 
 	// AActor interface.
@@ -38,7 +43,10 @@ protected:
 
 private:
 
-	UPROPERTY(VisibleAnywhere, Category=Camera, BlueprintGetter="GetCameraRigComponent", meta=(ExposeFunctionCategories="Camera"))
+	UPROPERTY(VisibleAnywhere, Category=GameplayCamera, BlueprintGetter="GetCameraRigComponent", meta=(ExposeFunctionCategories="Camera"))
 	TObjectPtr<UGameplayCameraRigComponent> CameraRigComponent;
+
+	UPROPERTY(Transient, VisibleAnywhere, Category=OutputCamera, BlueprintGetter="GetOutputCameraComponent")
+	TObjectPtr<UCineCameraComponent> OutputCameraComponent;
 };
 
