Cameras: add a sprite icon to GPC components

This makes it easier to locate the GPC component in the scene when the output camera is offset from it. The sprite is hidden when the output camera is close or on top of the root location.

#virtualized
#Synced-CL 51024716

diff --git a/Engine/Plugins/Cameras/GameplayCameras/Source/GameplayCameras/Private/GameFramework/GameplayCameraComponent.cpp b/Engine/Plugins/Cameras/GameplayCameras/Source/GameplayCameras/Private/GameFramework/GameplayCameraComponent.cpp
--- a/Engine/Plugins/Cameras/GameplayCameras/Source/GameplayCameras/Private/GameFramework/GameplayCameraComponent.cpp
+++ b/Engine/Plugins/Cameras/GameplayCameras/Source/GameplayCameras/Private/GameFramework/GameplayCameraComponent.cpp
@@ -53,10 +53,10 @@ void UGameplayCameraComponent::OnRegister()
 	Super::OnRegister();
 
 #if WITH_EDITOR
-
 	FGameplayCamerasDelegates::OnCameraAssetBuilt().AddUObject(this, &UGameplayCameraComponent::OnCameraAssetBuilt);
 
-#endif
+	Super::CreateCameraSpriteComponent(TEXT("/GameplayCameras/Textures/S_GameplayCamera.S_GameplayCamera"));
+#endif  // WITH_EDITOR
 }
 
 void UGameplayCameraComponent::OnUnregister()
@@ -64,9 +64,7 @@ void UGameplayCameraComponent::OnUnregister()
 	using namespace UE::Cameras;
 
 #if WITH_EDITOR
-
 	FGameplayCamerasDelegates::OnCameraAssetBuilt().RemoveAll(this);
-
 #endif  // WITH_EDITOR
 
 	Super::OnUnregister();
diff --git a/Engine/Plugins/Cameras/GameplayCameras/Source/GameplayCameras/Private/GameFramework/GameplayCameraComponentBase.cpp b/Engine/Plugins/Cameras/GameplayCameras/Source/GameplayCameras/Private/GameFramework/GameplayCameraComponentBase.cpp
--- a/Engine/Plugins/Cameras/GameplayCameras/Source/GameplayCameras/Private/GameFramework/GameplayCameraComponentBase.cpp
+++ b/Engine/Plugins/Cameras/GameplayCameras/Source/GameplayCameras/Private/GameFramework/GameplayCameraComponentBase.cpp
@@ -3,6 +3,7 @@
 #include "GameFramework/GameplayCameraComponentBase.h"
 
 #include "CineCameraComponent.h"
+#include "Components/BillboardComponent.h"
 #include "Core/CameraAsset.h"
 #include "Core/CameraEvaluationContextStack.h"
 #include "Core/CameraRigAsset.h"
@@ -12,6 +13,7 @@
 #include "Directors/SingleCameraDirector.h"
 #include "Engine/Canvas.h"
 #include "Engine/EngineTypes.h"
+#include "Engine/Texture2D.h"
 #include "Engine/World.h"
 #include "IGameplayCamerasLiveEditManager.h"
 #include "IGameplayCamerasModule.h"
@@ -21,6 +23,7 @@
 #include "Misc/EngineVersionComparison.h"
 #include "SceneView.h"
 #include "ShowFlags.h"
+#include "UObject/ICookInfo.h"
 
 #if UE_VERSION_NEWER_THAN_OR_EQUAL(5,6,0)
 #include "PrimitiveDrawInterface.h"
@@ -516,6 +519,31 @@ void UGameplayCameraComponentBase::OnRegister()
 #endif  // WITH_EDITOR
 }
 
+#if WITH_EDITOR
+
+void UGameplayCameraComponentBase::CreateCameraSpriteComponent(const FString& SpriteTexturePath)
+{
+	UTexture2D* EditorSpriteTexture = nullptr;
+	{
+		FCookLoadScope EditorOnlyScope(ECookLoadType::EditorOnly);
+		EditorSpriteTexture = LoadObject<UTexture2D>(nullptr, SpriteTexturePath);
+	}
+
+	if (EditorSpriteTexture)
+	{
+		bVisualizeComponent = true;
+		CreateSpriteComponent(EditorSpriteTexture);
+	}
+
+	if (SpriteComponent)
+	{
+		SpriteComponent->SpriteInfo.Category = TEXT("Cameras");
+		SpriteComponent->SpriteInfo.DisplayName = NSLOCTEXT("SpriteCategory", "Cameras", "Cameras");
+	}
+}
+
+#endif  // WITH_EDITOR
+
 void UGameplayCameraComponentBase::BeginPlay()
 {
 	using namespace UE::Cameras;
@@ -600,16 +628,13 @@ void UGameplayCameraComponentBase::TickComponent(float DeltaTime, ELevelTick Tic
 	Super::TickComponent(DeltaTime, TickType, ThisTickFunction);
 
 #if WITH_EDITOR
-
 	// Make sure things are setup (or not) if we want to run the camera logic in editor (or not).
 	AutoManageEditorPreviewEvaluator();
-
 #endif  // WITH_EDITOR
 
 	UpdateEvaluationContext(false);
 
 #if WITH_EDITOR
-
 	if (bIsEditorWorld)
 	{
 		UpdateCameraSystemForEditorPreview(DeltaTime);
@@ -619,18 +644,19 @@ void UGameplayCameraComponentBase::TickComponent(float DeltaTime, ELevelTick Tic
 		UpdateCameraSystem(DeltaTime);
 		UpdateControlRotationIfNeeded();
 	}
-
 #else
-
 	if (bRunStandaloneCameraSystem)
 	{
 		UpdateCameraSystem(DeltaTime);
 		UpdateControlRotationIfNeeded();
 	}
-
 #endif  // WITH_EDITOR
 
 	UpdateOutputCameraComponent();
+
+#if WITH_EDITOR
+	UpdateCameraSpriteComponent();
+#endif  // WITH_EDITOR
 }
 
 void UGameplayCameraComponentBase::UpdateEvaluationContext(bool bForceApplyParameterOverrides)
@@ -655,9 +681,7 @@ void UGameplayCameraComponentBase::UpdateEvaluationContext(bool bForceApplyParam
 		OnUpdateEvaluationContext(bForceApplyParameterOverrides);
 
 #if WITH_EDITOR
-
 		EvaluationContext->UpdateForEditorPreview();
-
 #endif  // WITH_EDITOR
 	}
 }
@@ -938,6 +962,30 @@ void UGameplayCameraComponentBase::UpdateOutputCameraComponent()
 	}
 }
 
+#if WITH_EDITOR
+
+void UGameplayCameraComponentBase::UpdateCameraSpriteComponent()
+{
+	if (SpriteComponent)
+	{
+		if (OutputCameraComponent)
+		{
+			const FVector3d RootLocation = GetComponentLocation();
+			const FVector3d OutputLocation = OutputCameraComponent->GetComponentLocation();
+			const double RootToOutputDistance = FVector3d::Distance(RootLocation, OutputLocation);
+
+			const bool bShouldBeVisible = (RootToOutputDistance > EditorSpriteHiddenWhenOutputCameraWithinDistance);
+			SpriteComponent->SetIsTemporarilyHiddenInEditor(!bShouldBeVisible);
+		}
+		else
+		{
+			SpriteComponent->SetIsTemporarilyHiddenInEditor(false);
+		}
+	}
+}
+
+#endif  // WITH_EDITOR
+
 void UGameplayCameraComponentBase::OnUpdateTransform(EUpdateTransformFlags UpdateTransformFlags, ETeleportType Teleport)
 {
 	Super::OnUpdateTransform(UpdateTransformFlags, Teleport);
diff --git a/Engine/Plugins/Cameras/GameplayCameras/Source/GameplayCameras/Private/GameFramework/GameplayCameraRigComponent.cpp b/Engine/Plugins/Cameras/GameplayCameras/Source/GameplayCameras/Private/GameFramework/GameplayCameraRigComponent.cpp
--- a/Engine/Plugins/Cameras/GameplayCameras/Source/GameplayCameras/Private/GameFramework/GameplayCameraRigComponent.cpp
+++ b/Engine/Plugins/Cameras/GameplayCameras/Source/GameplayCameras/Private/GameFramework/GameplayCameraRigComponent.cpp
@@ -45,10 +45,10 @@ void UGameplayCameraRigComponent::OnRegister()
 	Super::OnRegister();
 
 #if WITH_EDITOR
-
 	FGameplayCamerasDelegates::OnCameraRigAssetBuilt().AddUObject(this, &UGameplayCameraRigComponent::OnCameraRigAssetBuilt);
 
-#endif
+	Super::CreateCameraSpriteComponent(TEXT("/GameplayCameras/Textures/S_GameplayCameraRig.S_GameplayCameraRig"));
+#endif  // WITH_EDITOR
 }
 
 void UGameplayCameraRigComponent::OnUnregister()
@@ -56,9 +56,7 @@ void UGameplayCameraRigComponent::OnUnregister()
 	using namespace UE::Cameras;
 
 #if WITH_EDITOR
-
 	FGameplayCamerasDelegates::OnCameraRigAssetBuilt().RemoveAll(this);
-
 #endif  // WITH_EDITOR
 
 	Super::OnUnregister();
diff --git a/Engine/Plugins/Cameras/GameplayCameras/Source/GameplayCameras/Public/GameFramework/GameplayCameraActorBase.h b/Engine/Plugins/Cameras/GameplayCameras/Source/GameplayCameras/Public/GameFramework/GameplayCameraActorBase.h
--- a/Engine/Plugins/Cameras/GameplayCameras/Source/GameplayCameras/Public/GameFramework/GameplayCameraActorBase.h
+++ b/Engine/Plugins/Cameras/GameplayCameras/Source/GameplayCameras/Public/GameFramework/GameplayCameraActorBase.h
@@ -17,7 +17,7 @@ class UGameplayCameraSystemHost;
  * a private instance of the camera system. It does this if no camera system was found attached
  * under the player controller.
  */
-UCLASS(BlueprintType, MinimalAPI, ClassGroup=Camera, HideCategories=(Input, Rendering))
+UCLASS(BlueprintType, MinimalAPI, ClassGroup=Camera)
 class AGameplayCameraActorBase : public AActor
 {
 	GENERATED_BODY()
diff --git a/Engine/Plugins/Cameras/GameplayCameras/Source/GameplayCameras/Public/GameFramework/GameplayCameraComponentBase.h b/Engine/Plugins/Cameras/GameplayCameras/Source/GameplayCameras/Public/GameFramework/GameplayCameraComponentBase.h
--- a/Engine/Plugins/Cameras/GameplayCameras/Source/GameplayCameras/Public/GameFramework/GameplayCameraComponentBase.h
+++ b/Engine/Plugins/Cameras/GameplayCameras/Source/GameplayCameras/Public/GameFramework/GameplayCameraComponentBase.h
@@ -37,9 +37,9 @@ class FGameplayCameraComponentEvaluationContext;
 /**
  * A component that can run a camera asset inside its own camera evaluation context.
  */
-UCLASS(Blueprintable, MinimalAPI, Abstract,
-		ClassGroup=Camera, 
-		HideCategories=(Mobility, Rendering, LOD), 
+UCLASS(Blueprintable, MinimalAPI, Abstract, ClassGroup=Camera,
+		HideCategories=(Object, ActorComponent, Physics, Rendering, Mobility, LOD),
+		PrioritizeCategories=(Camera, EditorCamera, Activation),
 		meta=(BlueprintSpawnableComponent))
 class UGameplayCameraComponentBase 
 	: public USceneComponent
@@ -231,6 +231,8 @@ protected:
 			const FCameraVariableTableAllocationInfo& VariableTableAllocationInfo,
 			const FCameraContextDataTableAllocationInfo& ContextDataTableAllocationInfo);
 	void RecreateEditorWorldEvaluationContext();
+
+	void CreateCameraSpriteComponent(const FString& SpriteTexturePath);
 #endif  // WITH_EDITOR
 
 private:
@@ -248,6 +250,7 @@ private:
 #if WITH_EDITOR
 	void AutoManageEditorPreviewEvaluator();
 	void OnEditorPreviewCameraRigIndexChanged();
+	void UpdateCameraSpriteComponent();
 #endif  // WITH_EDITOR
 
 public:
@@ -280,13 +283,20 @@ public:
 #if WITH_EDITORONLY_DATA
 
 	/** Whether to run this camera in editor. */
-	UPROPERTY(EditAnywhere, Category=Camera)
+	UPROPERTY(EditAnywhere, Category=EditorCamera)
 	bool bRunInEditor = true;
 
 	/** The camera rig to run in the editor. */
-	UPROPERTY(EditAnywhere, Category=Camera, meta=(EditCondition="bRunInEditor"))
+	UPROPERTY(EditAnywhere, Category=EditorCamera, meta=(EditCondition="bRunInEditor"))
 	int32 EditorPreviewCameraRigIndex = 0;
 
+	/**
+	 * When the output camera's location is within this distance of the root, hide the editor sprite.
+	 * Set this value to zero or negative to never hide the editor sprite.
+	 */
+	UPROPERTY(EditAnywhere, Category=EditorCamera)
+	float EditorSpriteHiddenWhenOutputCameraWithinDistance = 100.f;
+
 #endif  // WITH_EDITORONLY_DATA
 
 private:
