Cameras: add "playback mode" to the GPC components

This allows GPC components to be recorded by [REDACTED_PERSON]. The output camera component is also recorded, and upon playback, the new "playback mode" option can be enabled (similar to how physics should be disabled on recorded actors). This lets the output camera component be animated by the recorded keyframes, without the GPC system interfering.
#Synced-CL 51024716

diff --git a/Engine/Plugins/Cameras/GameplayCameras/Source/GameplayCameras/Private/GameFramework/GameplayCameraComponentBase.cpp b/Engine/Plugins/Cameras/GameplayCameras/Source/GameplayCameras/Private/GameFramework/GameplayCameraComponentBase.cpp
--- a/Engine/Plugins/Cameras/GameplayCameras/Source/GameplayCameras/Private/GameFramework/GameplayCameraComponentBase.cpp
+++ b/Engine/Plugins/Cameras/GameplayCameras/Source/GameplayCameras/Private/GameFramework/GameplayCameraComponentBase.cpp
@@ -120,17 +120,21 @@ bool UGameplayCameraComponentBase::CanRunCameraSystem() const
 {
 	using namespace UE::Cameras;
 
-#if WITH_EDITOR
-	IGameplayCamerasModule& GameplayCamerasModule = IGameplayCamerasModule::Get();
-	if (TSharedPtr<IGameplayCamerasLiveEditManager> LiveEditManager = GameplayCamerasModule.GetLiveEditManager())
+	if (!bPlaybackMode)
 	{
-		const bool bCanRunInEditor = bRunInEditor && GameplayCamerasModule.GetLiveEditManager()->CanRunInEditor();
-		return (!bIsEditorWorld || bCanRunInEditor);
-	}
-	return !bIsEditorWorld;
+#if WITH_EDITOR
+		IGameplayCamerasModule& GameplayCamerasModule = IGameplayCamerasModule::Get();
+		if (TSharedPtr<IGameplayCamerasLiveEditManager> LiveEditManager = GameplayCamerasModule.GetLiveEditManager())
+		{
+			const bool bCanRunInEditor = bRunInEditor && GameplayCamerasModule.GetLiveEditManager()->CanRunInEditor();
+			return (!bIsEditorWorld || bCanRunInEditor);
+		}
+		return !bIsEditorWorld;
 #else
-	return true;
+		return true;
 #endif  // WITH_EDITOR
+	}
+	return false;
 }
 
 bool UGameplayCameraComponentBase::EnsureCameraSystemHostIfNeeded()
@@ -889,7 +893,7 @@ void UGameplayCameraComponentBase::UpdateOutputCameraComponent()
 	// So we lazily do it here instead.
 	EnsureOutputCameraComponent();
 
-	if (!OutputCameraComponent)
+	if (!OutputCameraComponent || bPlaybackMode)
 	{
 		return;
 	}
diff --git a/Engine/Plugins/Cameras/GameplayCameras/Source/GameplayCameras/Public/GameFramework/GameplayCameraComponentBase.h b/Engine/Plugins/Cameras/GameplayCameras/Source/GameplayCameras/Public/GameFramework/GameplayCameraComponentBase.h
--- a/Engine/Plugins/Cameras/GameplayCameras/Source/GameplayCameras/Public/GameFramework/GameplayCameraComponentBase.h
+++ b/Engine/Plugins/Cameras/GameplayCameras/Source/GameplayCameras/Public/GameFramework/GameplayCameraComponentBase.h
@@ -280,6 +280,14 @@ public:
 	UPROPERTY(BlueprintReadWrite, EditAnywhere, Category=Camera, meta=(EditCondition="bRunStandaloneCameraSystem"))
 	bool bSetControlRotationWhenViewTarget = false;
 
+	/**
+	 * Enables or disables playback mode, which turns of any standalone camera system evaluation
+	 * and lets the output camera component be manipulated externally. This is useful when the
+	 * evaluated camera was recorded (e.g. with Take Recorder) and turned into keyframed animation.
+	 */
+	UPROPERTY(BlueprintReadWrite, EditAnywhere, Category=Camera)
+	bool bPlaybackMode = false;
+
 #if WITH_EDITORONLY_DATA
 
 	/** Whether to run this camera in editor. */
