Cameras: rename AutoActivateForPlayer to DefaultPlayer, and simply always use that unless anything else is specified

#Synced-CL 51024716

diff --git a/Engine/Plugins/Cameras/GameplayCameras/Config/DefaultGameplayCameras.ini b/Engine/Plugins/Cameras/GameplayCameras/Config/DefaultGameplayCameras.ini
--- a/Engine/Plugins/Cameras/GameplayCameras/Config/DefaultGameplayCameras.ini
+++ b/Engine/Plugins/Cameras/GameplayCameras/Config/DefaultGameplayCameras.ini
@@ -39,3 +39,5 @@
 +PropertyRedirects=(OldName="/Script/GameplayCameras.CustomCameraNodeBlendableParameter.ParameterTypeObject",NewName="DataTypeObject")
 +PropertyRedirects=(OldName="/Script/GameplayCameras.CameraObjectInterfaceBlendableParameter.ParameterType",NewName="VariableType")
 
++PropertyRedirects=(OldName="/Script/GameplayCameras.GameplayCameraComponentBase.AutoActivateForPlayer",NewName="DefaultPlayer")
+
diff --git a/Engine/Plugins/Cameras/GameplayCameras/Source/GameplayCameras/Private/GameFramework/GameplayCameraComponentBase.cpp b/Engine/Plugins/Cameras/GameplayCameras/Source/GameplayCameras/Private/GameFramework/GameplayCameraComponentBase.cpp
--- a/Engine/Plugins/Cameras/GameplayCameras/Source/GameplayCameras/Private/GameFramework/GameplayCameraComponentBase.cpp
+++ b/Engine/Plugins/Cameras/GameplayCameras/Source/GameplayCameras/Private/GameFramework/GameplayCameraComponentBase.cpp
@@ -95,6 +95,12 @@ void UGameplayCameraComponentBase::ActivateCameraForPlayerController(APlayerCont
 	// our private camera system.
 	Super::Activate(false);
 
+	if (!PlayerController && DefaultPlayer != EAutoReceiveInput::Disabled)
+	{
+		const int32 PlayerIndex = DefaultPlayer.GetIntValue() - 1;
+		PlayerController = UGameplayStatics::GetPlayerController(this, PlayerIndex);
+	}
+
 	EnsureEvaluationContext(PlayerController);
 	EnsureCameraSystemHostIfNeeded();
 
@@ -560,17 +566,9 @@ void UGameplayCameraComponentBase::BeginPlay()
 	// However, it can happen that some BP construction script already called ActivateCameraForXyz()
 	// before we got to start play (e.g. from a parent actor) and so in this case, let's skip
 	// re-activating for nothing.
-	if (IsActive() && !EvaluationContext)
+	if (IsActive() && !EvaluationContext && GetNetMode() != NM_DedicatedServer)
 	{
-		if (AutoActivateForPlayer != EAutoReceiveInput::Disabled && GetNetMode() != NM_DedicatedServer)
-		{
-			const int32 PlayerIndex = AutoActivateForPlayer.GetIntValue() - 1;
-			ActivateCameraForPlayerIndex(PlayerIndex);
-		}
-		else
-		{
-			ActivateCameraForPlayerController(nullptr, false);
-		}
+		ActivateCameraForPlayerController(nullptr, false);
 	}
 }
 
diff --git a/Engine/Plugins/Cameras/GameplayCameras/Source/GameplayCameras/Public/GameFramework/GameplayCameraComponentBase.h b/Engine/Plugins/Cameras/GameplayCameras/Source/GameplayCameras/Public/GameFramework/GameplayCameraComponentBase.h
--- a/Engine/Plugins/Cameras/GameplayCameras/Source/GameplayCameras/Public/GameFramework/GameplayCameraComponentBase.h
+++ b/Engine/Plugins/Cameras/GameplayCameras/Source/GameplayCameras/Public/GameFramework/GameplayCameraComponentBase.h
@@ -7,6 +7,7 @@
 #include "Core/CameraEvaluationContext.h"
 #include "Core/CameraRigInstanceID.h"
 #include "Core/CameraShakeInstanceID.h"
+#include "Engine/EngineTypes.h"
 #include "GameFramework/BlueprintCameraEvaluationDataRef.h"
 #include "GameFramework/IGameplayCameraSystemHost.h"
 #include "Services/CameraActionInstanceID.h"
@@ -256,15 +257,18 @@ private:
 public:
 
 	/**
-	 * If AutoActivate is set, auto-activates this component's camera for the given player.
-	 * This is equivalent to calling ActivateCameraForPlayerIndex on BeginPlay.
+	 * The player controller to bind to by default.
+	 *
+	 * Auto-activation and explicit calls to Activate will use this player, if valid. Calls to 
+	 * ActivateCameraForPlayerIndex and ActivateCameraForPlayerController will use their given player, 
+	 * but fall back to this player if they are given an invalid parameter.
 	 */
-	UPROPERTY(BlueprintReadWrite, EditAnywhere, Category=Activation, meta=(EditCondition="bAutoActivate"))
-	TEnumAsByte<EAutoReceiveInput::Type> AutoActivateForPlayer;
+	UPROPERTY(BlueprintReadWrite, EditAnywhere, Category=Activation)
+	TEnumAsByte<EAutoReceiveInput::Type> DefaultPlayer = EAutoReceiveInput::Player0;
 
 	/**
 	 * When enabled, this camera component runs its own camera system when activated, which
-	 * makes it posisble for it to run independently. When disabled, this camera component
+	 * makes it possible for it to run independently. When disabled, this camera component
 	 * only creates an evaluation context that must be run with an external camera system,
 	 * such as the Gameplay Cameras Player Camera Manager.
 	 */
