Cameras: some errors and warnings shouldn't be reported when running editor previews

#Synced-CL 51024716

diff --git a/Engine/Plugins/Cameras/GameplayCameras/Source/GameplayCameras/Private/Nodes/Attach/AttachToPlayerPawnCameraNode.cpp b/Engine/Plugins/Cameras/GameplayCameras/Source/GameplayCameras/Private/Nodes/Attach/AttachToPlayerPawnCameraNode.cpp
--- a/Engine/Plugins/Cameras/GameplayCameras/Source/GameplayCameras/Private/Nodes/Attach/AttachToPlayerPawnCameraNode.cpp
+++ b/Engine/Plugins/Cameras/GameplayCameras/Source/GameplayCameras/Private/Nodes/Attach/AttachToPlayerPawnCameraNode.cpp
@@ -48,8 +48,11 @@ void FAttachToPlayerPawnCameraNodeEvaluator::OnInitialize(const FCameraNodeEvalu
 	APlayerController* PlayerController = Params.EvaluationContext->GetPlayerController();
 	if (!PlayerController)
 	{
-		UE_LOG(LogCameraSystem, Error, 
-				TEXT("Can't run AttatchToPlayerPawn camera node because no player controller was found on the context."));
+		if (Params.Evaluator->GetRole() == ECameraSystemEvaluatorRole::Game)
+		{
+			UE_LOG(LogCameraSystem, Error, 
+					TEXT("Can't run AttatchToPlayerPawn camera node because no player controller was found on the context."));
+		}
 		bHasValidPlayerController = false;
 		return;
 	}
diff --git a/Engine/Plugins/Cameras/GameplayCameras/Source/GameplayCameras/Private/Nodes/Input/DrivenControlRotationCameraNode.cpp b/Engine/Plugins/Cameras/GameplayCameras/Source/GameplayCameras/Private/Nodes/Input/DrivenControlRotationCameraNode.cpp
--- a/Engine/Plugins/Cameras/GameplayCameras/Source/GameplayCameras/Private/Nodes/Input/DrivenControlRotationCameraNode.cpp
+++ b/Engine/Plugins/Cameras/GameplayCameras/Source/GameplayCameras/Private/Nodes/Input/DrivenControlRotationCameraNode.cpp
@@ -89,13 +89,6 @@ void FDrivenControlRotationHelperService::AddMonitoredContext(TSharedPtr<const F
 		NewEntry.CachedControlRotation = PlayerController->GetControlRotation();
 		NewEntry.MonitoringRequests = 1;
 	}
-	else
-	{
-		UE_LOG(LogCameraSystem, Warning, 
-				TEXT("Can't monitor control rotation for camera context owned by '%s', it is running without "
-					 "any association to a player controller."),
-				*GetNameSafe(InContext->GetOwner()));
-	}
 }
 
 void FDrivenControlRotationHelperService::RemoveMonitoredContext(TSharedPtr<const FCameraEvaluationContext> InContext)
