Cameras: make playback mode only affect whether a GPC component writes to the output component

Don't make it affect whether a standalone camera system is created or not. This makes it easier to manage, and easier to understand for the user.

#Synced-CL 51361382

diff --git a/Engine/Plugins/Cameras/GameplayCameras/Source/GameplayCameras/Private/GameFramework/GameplayCameraComponentBase.cpp b/Engine/Plugins/Cameras/GameplayCameras/Source/GameplayCameras/Private/GameFramework/GameplayCameraComponentBase.cpp
--- a/Engine/Plugins/Cameras/GameplayCameras/Source/GameplayCameras/Private/GameFramework/GameplayCameraComponentBase.cpp
+++ b/Engine/Plugins/Cameras/GameplayCameras/Source/GameplayCameras/Private/GameFramework/GameplayCameraComponentBase.cpp
@@ -126,21 +126,17 @@ bool UGameplayCameraComponentBase::CanRunCameraSystem() const
 {
 	using namespace UE::Cameras;
 
-	if (!bPlaybackMode)
-	{
 #if WITH_EDITOR
-		IGameplayCamerasModule& GameplayCamerasModule = IGameplayCamerasModule::Get();
-		if (TSharedPtr<IGameplayCamerasLiveEditManager> LiveEditManager = GameplayCamerasModule.GetLiveEditManager())
-		{
-			const bool bCanRunInEditor = bRunInEditor && GameplayCamerasModule.GetLiveEditManager()->CanRunInEditor();
-			return (!bIsEditorWorld || bCanRunInEditor);
-		}
-		return !bIsEditorWorld;
+	IGameplayCamerasModule& GameplayCamerasModule = IGameplayCamerasModule::Get();
+	if (TSharedPtr<IGameplayCamerasLiveEditManager> LiveEditManager = GameplayCamerasModule.GetLiveEditManager())
+	{
+		const bool bCanRunInEditor = bRunInEditor && GameplayCamerasModule.GetLiveEditManager()->CanRunInEditor();
+		return (!bIsEditorWorld || bCanRunInEditor);
+	}
+	return !bIsEditorWorld;
 #else
-		return true;
+	return true;
 #endif  // WITH_EDITOR
-	}
-	return false;
 }
 
 bool UGameplayCameraComponentBase::EnsureCameraSystemHostIfNeeded()
